Avoid ticking UContentBrowserDataSubsystem when it may be unsafe to do so

diff --git a/Engine/Source/Editor/ContentBrowserData/Private/ContentBrowserDataSubsystem.cpp b/Engine/Source/Editor/ContentBrowserData/Private/ContentBrowserDataSubsystem.cpp
--- a/Engine/Source/Editor/ContentBrowserData/Private/ContentBrowserDataSubsystem.cpp
+++ b/Engine/Source/Editor/ContentBrowserData/Private/ContentBrowserDataSubsystem.cpp
@@ -6,6 +6,7 @@
 #include "Misc/PackageName.h"
 #include "Features/IModularFeatures.h"
 #include "Stats/Stats.h"
+#include "UObject/UObjectThreadContext.h"
 #include "Framework/Application/SlateApplication.h"
 
 void UContentBrowserDataSubsystem::Initialize(FSubsystemCollectionBase& Collection)
@@ -558,6 +559,13 @@ void UContentBrowserDataSubsystem::HandleDataSourceUnregistered(const FName& Typ
 void UContentBrowserDataSubsystem::Tick(const float InDeltaTime)
 {
 	QUICK_SCOPE_CYCLE_COUNTER(STAT_UContentBrowserDataSubsystem_Tick);
+
+	if (GIsSavingPackage || IsGarbageCollecting() || FUObjectThreadContext::Get().IsRoutingPostLoad)
+	{
+		// Not to safe to Tick right now, as the below code may try and find objects
+		return;
+	}
+
 	for (const auto& AvailableDataSourcePair : AvailableDataSources)
 	{
 		AvailableDataSourcePair.Value->Tick(InDeltaTime);
