[ContentBrowser] Fix template shared paths not showing up in content browser

diff --git a/Engine/Source/Editor/ContentBrowserData/Private/ContentBrowserDataSource.cpp b/Engine/Source/Editor/ContentBrowserData/Private/ContentBrowserDataSource.cpp
--- a/Engine/Source/Editor/ContentBrowserData/Private/ContentBrowserDataSource.cpp
+++ b/Engine/Source/Editor/ContentBrowserData/Private/ContentBrowserDataSource.cpp
@@ -156,6 +156,25 @@ void UContentBrowserDataSource::RootPathAdded(const FStringView InInternalPath)
 		Path.LeftChopInline(1);
 	}
 
+	// Handle root containing multiple folders such as /Game/FirstPerson
+	int32 SecondSlashIndex = INDEX_NONE;
+	if (Path.Len() > 1 && Path.RightChop(1).FindChar(TEXT('/'), SecondSlashIndex))
+	{
+		const FStringView FirstPath = Path.Left(SecondSlashIndex + 1);
+
+		FName FirstParentVirtualPath;
+		TryConvertInternalPathToVirtual(FName(FirstPath), FirstParentVirtualPath);
+
+		bool bIsFullyVirtual = false;
+		if (RootPathVirtualTree.PathExists(FirstParentVirtualPath, bIsFullyVirtual))
+		{
+			if (bIsFullyVirtual == false)
+			{
+				return;
+			}
+		}
+	}
+
 	FName VirtualPath;
 	FName PathFName(Path);
 	TryConvertInternalPathToVirtual(PathFName, VirtualPath);
