Suppress Tick from modal dialogs during rename finalization in the Content Browser

This can cause a reentrancy issue where the item currently being renamed is removed via an update while we're still processing it

diff --git a/Engine/Source/Editor/ContentBrowser/Private/SAssetView.cpp b/Engine/Source/Editor/ContentBrowser/Private/SAssetView.cpp
--- a/Engine/Source/Editor/ContentBrowser/Private/SAssetView.cpp
+++ b/Engine/Source/Editor/ContentBrowser/Private/SAssetView.cpp
@@ -671,6 +671,9 @@ FContentBrowserItem SAssetView::EndCreateDeferredItem(const TSharedPtr<FAssetVie
 		// If not finalizing then we just discard the temporary
 		if (bFinalize)
 		{
+			UContentBrowserDataSubsystem* ContentBrowserData = IContentBrowserDataModule::Get().GetSubsystem();
+			FScopedSuppressContentBrowserDataTick TickSuppression(ContentBrowserData);
+
 			if (DeferredItemToCreate->ItemContext.ValidateItem(InName, &OutErrorText))
 			{
 				FinalizedItem = DeferredItemToCreate->ItemContext.FinalizeItem(InName, &OutErrorText);
@@ -3693,6 +3696,9 @@ void SAssetView::AssetRenameCommit(const TSharedPtr<FAssetViewItem>& Item, const
 	}
 	else if (CommitType != ETextCommit::OnCleared && !Item->GetItem().GetItemName().ToString().Equals(NewName))
 	{
+		UContentBrowserDataSubsystem* ContentBrowserData = IContentBrowserDataModule::Get().GetSubsystem();
+		FScopedSuppressContentBrowserDataTick TickSuppression(ContentBrowserData);
+
 		FContentBrowserItem NewItem;
 		if (Item->GetItem().CanRename(&NewName, &ErrorMessage) && Item->GetItem().Rename(NewName, &NewItem))
 		{
diff --git a/Engine/Source/Editor/ContentBrowser/Private/SPathView.cpp b/Engine/Source/Editor/ContentBrowser/Private/SPathView.cpp
--- a/Engine/Source/Editor/ContentBrowser/Private/SPathView.cpp
+++ b/Engine/Source/Editor/ContentBrowser/Private/SPathView.cpp
@@ -1591,6 +1591,9 @@ void SPathView::FolderNameChanged( const TSharedPtr< FTreeItem >& TreeItem, cons
 		}
 		else
 		{
+			UContentBrowserDataSubsystem* ContentBrowserData = IContentBrowserDataModule::Get().GetSubsystem();
+			FScopedSuppressContentBrowserDataTick TickSuppression(ContentBrowserData);
+
 			if (PendingNewFolderContext.ValidateItem(ProposedName, &ErrorMessage))
 			{
 				NewItem = PendingNewFolderContext.FinalizeItem(ProposedName, &ErrorMessage);
@@ -1605,6 +1608,9 @@ void SPathView::FolderNameChanged( const TSharedPtr< FTreeItem >& TreeItem, cons
 	}
 	else if (CommitType != ETextCommit::OnCleared && !TreeItem->GetItem().GetItemName().ToString().Equals(ProposedName))
 	{
+		UContentBrowserDataSubsystem* ContentBrowserData = IContentBrowserDataModule::Get().GetSubsystem();
+		FScopedSuppressContentBrowserDataTick TickSuppression(ContentBrowserData);
+
 		if (TreeItem->GetItem().CanRename(&ProposedName, &ErrorMessage) && TreeItem->GetItem().Rename(ProposedName, &NewItem))
 		{
 			bSuccess = true;
diff --git a/Engine/Source/Editor/ContentBrowserData/Private/ContentBrowserDataSubsystem.cpp b/Engine/Source/Editor/ContentBrowserData/Private/ContentBrowserDataSubsystem.cpp
--- a/Engine/Source/Editor/ContentBrowserData/Private/ContentBrowserDataSubsystem.cpp
+++ b/Engine/Source/Editor/ContentBrowserData/Private/ContentBrowserDataSubsystem.cpp
@@ -544,6 +544,12 @@ void UContentBrowserDataSubsystem::Tick(const float InDeltaTime)
 		return;
 	}
 
+	if (TickSuppressionCount > 0)
+	{
+		// Not safe to Tick right now, as we've been asked not to
+		return;
+	}
+
 	for (const auto& AvailableDataSourcePair : AvailableDataSources)
 	{
 		AvailableDataSourcePair.Value->Tick(InDeltaTime);
diff --git a/Engine/Source/Editor/ContentBrowserData/Public/ContentBrowserDataSubsystem.h b/Engine/Source/Editor/ContentBrowserData/Public/ContentBrowserDataSubsystem.h
--- a/Engine/Source/Editor/ContentBrowserData/Public/ContentBrowserDataSubsystem.h
+++ b/Engine/Source/Editor/ContentBrowserData/Public/ContentBrowserDataSubsystem.h
@@ -61,6 +61,8 @@ class CONTENTBROWSERDATA_API UContentBrowserDataSubsystem : public UEditorSubsys
 	GENERATED_BODY()
 
 public:
+	friend class FScopedSuppressContentBrowserDataTick;
+
 	//~ UEditorSubsystem interface
 	virtual void Initialize(FSubsystemCollectionBase& Collection) override;
 	virtual void Deinitialize() override;
@@ -336,6 +338,12 @@ private:
 	 */
 	bool bPendingItemDataRefreshedNotification = false;
 
+	/**
+	 * >0 if Tick events have currently been suppressed.
+	 * @see FScopedSuppressContentBrowserDataTick
+	 */
+	int32 TickSuppressionCount = 0;
+
 	/**
 	 * Delegate called for incremental item data updates from data sources that can provide delta-updates.
 	 */
@@ -372,3 +380,29 @@ private:
 	 */
 	FString AllFolderPrefix;
 };
+
+/**
+ * Helper to suppress Tick events during critical times, when the underlying data should not be updated.
+ */
+class FScopedSuppressContentBrowserDataTick
+{
+public:
+	explicit FScopedSuppressContentBrowserDataTick(UContentBrowserDataSubsystem* InContentBrowserData)
+		: ContentBrowserData(InContentBrowserData)
+	{
+		check(ContentBrowserData);
+		++ContentBrowserData->TickSuppressionCount;
+	}
+
+	~FScopedSuppressContentBrowserDataTick()
+	{
+		checkf(ContentBrowserData->TickSuppressionCount > 0, TEXT("TickSuppressionCount underflow!"));
+		--ContentBrowserData->TickSuppressionCount;
+	}
+
+	FScopedSuppressContentBrowserDataTick(const FScopedSuppressContentBrowserDataTick&) = delete;
+	FScopedSuppressContentBrowserDataTick& operator=(const FScopedSuppressContentBrowserDataTick&) = delete;
+
+private:
+	UContentBrowserDataSubsystem* ContentBrowserData = nullptr;
+};
