Ensure that the array of FContentBrowserItemDataUpdate can't be invalidated mid-update

Speculative fix for an issue of bogus data being passed into the callback, likely from the array being appended to mid-update

diff --git a/Engine/Source/Editor/ContentBrowserData/Private/ContentBrowserDataSubsystem.cpp b/Engine/Source/Editor/ContentBrowserData/Private/ContentBrowserDataSubsystem.cpp
--- a/Engine/Source/Editor/ContentBrowserData/Private/ContentBrowserDataSubsystem.cpp
+++ b/Engine/Source/Editor/ContentBrowserData/Private/ContentBrowserDataSubsystem.cpp
@@ -563,8 +563,9 @@ void UContentBrowserDataSubsystem::Tick(const float InDeltaTime)
 
 	if (PendingUpdates.Num() > 0)
 	{
-		ItemDataUpdatedDelegate.Broadcast(MakeArrayView(PendingUpdates));
+		TArray<FContentBrowserItemDataUpdate> LocalPendingUpdates = MoveTemp(PendingUpdates);
 		PendingUpdates.Empty();
+		ItemDataUpdatedDelegate.Broadcast(MakeArrayView(LocalPendingUpdates));
 	}
 
 	if (ActiveDataSourcesDiscoveringContent.Num() > 0)
