Reduce content browser hitches during PIE by delaying modification updates until after PIE finishes.

diff --git a/Engine/Source/Editor/ContentBrowserData/Private/ContentBrowserDataSubsystem.cpp b/Engine/Source/Editor/ContentBrowserData/Private/ContentBrowserDataSubsystem.cpp
--- a/Engine/Source/Editor/ContentBrowserData/Private/ContentBrowserDataSubsystem.cpp
+++ b/Engine/Source/Editor/ContentBrowserData/Private/ContentBrowserDataSubsystem.cpp
@@ -11,6 +11,7 @@
 #include "Settings/ContentBrowserSettings.h"
 #include "Interfaces/IPluginManager.h"
 #include "Framework/Application/SlateApplication.h"
+#include "Editor.h"
 
 void UContentBrowserDataSubsystem::Initialize(FSubsystemCollectionBase& Collection)
 {
@@ -38,6 +39,9 @@ void UContentBrowserDataSubsystem::Initialize(FSubsystemCollectionBase& Collecti
 	ModularFeatures.OnModularFeatureRegistered().AddUObject(this, &UContentBrowserDataSubsystem::HandleDataSourceRegistered);
 	ModularFeatures.OnModularFeatureUnregistered().AddUObject(this, &UContentBrowserDataSubsystem::HandleDataSourceUnregistered);
 
+	FEditorDelegates::BeginPIE.AddUObject(this, &UContentBrowserDataSubsystem::OnBeginPIE);
+	FEditorDelegates::EndPIE.AddUObject(this, &UContentBrowserDataSubsystem::OnEndPIE);
+
 	// Tick during normal operation
 	TickHandle = FTicker::GetCoreTicker().AddTicker(TEXT("ContentBrowserData"), 0.1f, [this](const float InDeltaTime)
 	{
@@ -58,6 +62,9 @@ void UContentBrowserDataSubsystem::Deinitialize()
 	ModularFeatures.OnModularFeatureRegistered().RemoveAll(this);
 	ModularFeatures.OnModularFeatureUnregistered().RemoveAll(this);
 
+	FEditorDelegates::BeginPIE.RemoveAll(this);
+	FEditorDelegates::EndPIE.RemoveAll(this);
+
 	ActiveDataSources.Reset();
 	AvailableDataSources.Reset();
 	ActiveDataSourcesDiscoveringContent.Reset();
@@ -557,6 +564,7 @@ void UContentBrowserDataSubsystem::Tick(const float InDeltaTime)
 	if (bPendingItemDataRefreshedNotification)
 	{
 		bPendingItemDataRefreshedNotification = false;
+		bHasIgnoredItemUpdates = false;
 		PendingUpdates.Empty();
 		ItemDataRefreshedDelegate.Broadcast();
 	}
@@ -598,8 +606,16 @@ void UContentBrowserDataSubsystem::Tick(const float InDeltaTime)
 
 void UContentBrowserDataSubsystem::QueueItemDataUpdate(FContentBrowserItemDataUpdate&& InUpdate)
 {
-	// TODO: Merge multiple Modified updates for a single item?
-	PendingUpdates.Emplace(MoveTemp(InUpdate));
+	// Ignore modified during PIE to reduce hitches, they will be updated when PIE stops
+	if ((InUpdate.GetUpdateType() == EContentBrowserItemUpdateType::Modified) && !AllowModifiedItemDataUpdates())
+	{
+		bHasIgnoredItemUpdates = true;
+	}
+	else
+	{
+		// TODO: Merge multiple Modified updates for a single item?
+		PendingUpdates.Emplace(MoveTemp(InUpdate));
+	}
 }
 
 void UContentBrowserDataSubsystem::NotifyItemDataRefreshed()
@@ -607,6 +623,27 @@ void UContentBrowserDataSubsystem::NotifyItemDataRefreshed()
 	bPendingItemDataRefreshedNotification = true;
 }
 
+bool UContentBrowserDataSubsystem::AllowModifiedItemDataUpdates() const
+{
+	return !bIsPIEActive;
+}
+
+void UContentBrowserDataSubsystem::OnBeginPIE(const bool bIsSimulating)
+{
+	bIsPIEActive = true;
+}
+
+void UContentBrowserDataSubsystem::OnEndPIE(const bool bIsSimulating)
+{
+	bIsPIEActive = false;
+
+	if (bHasIgnoredItemUpdates)
+	{
+		// Perform a full update because modified updates were ignored during PIE
+		NotifyItemDataRefreshed();
+	}
+}
+
 void UContentBrowserDataSubsystem::SetPathViewSpecialSortFolders(const TArray<FName>& InSpecialSortFolders)
 {
 	PathViewSpecialSortFolders = InSpecialSortFolders;
diff --git a/Engine/Source/Editor/ContentBrowserData/Public/ContentBrowserDataSubsystem.h b/Engine/Source/Editor/ContentBrowserData/Public/ContentBrowserDataSubsystem.h
--- a/Engine/Source/Editor/ContentBrowserData/Public/ContentBrowserDataSubsystem.h
+++ b/Engine/Source/Editor/ContentBrowserData/Public/ContentBrowserDataSubsystem.h
@@ -296,6 +296,21 @@ private:
 	 */
 	void Tick(const float InDeltaTime);
 
+	/**
+	 * Returns true if item data modifications are being processed.
+	 */
+	bool AllowModifiedItemDataUpdates() const;
+
+	/**
+	 * Called when Play in Editor begins.
+	 */
+	void OnBeginPIE(const bool bIsSimulating);
+
+	/**
+	 * Called when Play in Editor stops.
+	 */
+	void OnEndPIE(const bool bIsSimulating);
+
 	//~ IContentBrowserItemDataSink interface
 	virtual void QueueItemDataUpdate(FContentBrowserItemDataUpdate&& InUpdate) override;
 	virtual void NotifyItemDataRefreshed() override;
@@ -338,6 +353,16 @@ private:
 	 */
 	bool bPendingItemDataRefreshedNotification = false;
 
+	/**
+	 * True if there are currently any ignored changes.
+	 */
+	bool bHasIgnoredItemUpdates = false;
+
+	/**
+	 * True if Play in Editor is active.
+	 */
+	bool bIsPIEActive = false;
+
 	/**
 	 * >0 if Tick events have currently been suppressed.
 	 * @see FScopedSuppressContentBrowserDataTick
