[ContentBrowser] Fix 'EngineData' folder showing up in breadcrumb instead of 'Engine'
#rnx

diff --git a/Engine/Source/Editor/ContentBrowserData/Private/ContentBrowserDataSubsystem.cpp b/Engine/Source/Editor/ContentBrowserData/Private/ContentBrowserDataSubsystem.cpp
--- a/Engine/Source/Editor/ContentBrowserData/Private/ContentBrowserDataSubsystem.cpp
+++ b/Engine/Source/Editor/ContentBrowserData/Private/ContentBrowserDataSubsystem.cpp
@@ -320,8 +320,7 @@ void UContentBrowserDataSubsystem::EnumerateItemsAtPath(const FName InPath, cons
 		{
 			if (!bHandledVirtualFolder && EnumHasAnyFlags(InItemTypeFilter, EContentBrowserItemTypeFilter::IncludeFolders))
 			{
-				const FString LeafName = FPackageName::GetShortName(InPath);
-				InCallback(FContentBrowserItemData(DataSource, EContentBrowserItemFlags::Type_Folder, InPath, *LeafName, FText(), nullptr));
+				InCallback(DataSource->CreateVirtualFolderItem(InPath));
 				bHandledVirtualFolder = true;
 			}
 		}
