[ContentBrowser] Add debug cvar commands for testing path conversion
#rnx

diff --git a/Engine/Source/Editor/ContentBrowserData/Private/ContentBrowserDataSubsystem.cpp b/Engine/Source/Editor/ContentBrowserData/Private/ContentBrowserDataSubsystem.cpp
--- a/Engine/Source/Editor/ContentBrowserData/Private/ContentBrowserDataSubsystem.cpp
+++ b/Engine/Source/Editor/ContentBrowserData/Private/ContentBrowserDataSubsystem.cpp
@@ -2,6 +2,7 @@
 
 #include "ContentBrowserDataSubsystem.h"
 #include "ContentBrowserDataSource.h"
+#include "IContentBrowserDataModule.h"
 #include "Containers/Ticker.h"
 #include "Misc/PackageName.h"
 #include "Misc/Paths.h"
@@ -13,6 +14,41 @@
 #include "Framework/Application/SlateApplication.h"
 #include "Editor.h"
 
+DEFINE_LOG_CATEGORY_STATIC(LogContentBrowserDataSubsystem, Log, All);
+
+namespace ContentBrowserDataSubsystem
+{
+	FAutoConsoleCommand CVarContentBrowserDebug_TryConvertVirtualPath = FAutoConsoleCommand(
+		TEXT("ContentBrowser.Debug.TryConvertVirtualPath"),
+		TEXT("Try to convert virtual path"),
+		FConsoleCommandWithArgsDelegate::CreateLambda([](const TArray<FString>& Args)
+		{
+			if (Args.Num() > 0)
+			{
+				const FString& VirtualPath = Args[0];
+				FNameBuilder ConvertedPath;
+				EContentBrowserPathType PathType = IContentBrowserDataModule::Get().GetSubsystem()->TryConvertVirtualPath(VirtualPath, ConvertedPath);
+				UE_LOG(LogContentBrowserDataSubsystem, Log, TEXT("InputVirtualPath: %s, ConvertedPath: %s, ConvertedPathType: %s"), *VirtualPath, *ConvertedPath, *UEnum::GetValueAsString(PathType));
+			}
+		}
+	));
+
+	FAutoConsoleCommand CVarContentBrowserDebug_ConvertInternalPathToVirtual = FAutoConsoleCommand(
+		TEXT("ContentBrowser.Debug.ConvertInternalPathToVirtual"),
+		TEXT("Convert internal path"),
+		FConsoleCommandWithArgsDelegate::CreateLambda([](const TArray<FString>& Args)
+		{
+			if (Args.Num() > 0)
+			{
+				const FString& InternalPath = Args[0];
+				FNameBuilder ConvertedPath;
+				IContentBrowserDataModule::Get().GetSubsystem()->ConvertInternalPathToVirtual(InternalPath, ConvertedPath);
+				UE_LOG(LogContentBrowserDataSubsystem, Log, TEXT("InputInternalPath: %s, ConvertedVirtualPath: %s"), *InternalPath, *ConvertedPath);
+			}
+		}
+	));
+}
+
 void UContentBrowserDataSubsystem::Initialize(FSubsystemCollectionBase& Collection)
 {
 	AllFolderPrefix = TEXT("/All");
