Expose FContentBrowserItemPath to scripting

diff --git a/Engine/Source/Editor/ContentBrowserData/Private/ContentBrowserDataSubsystem.cpp b/Engine/Source/Editor/ContentBrowserData/Private/ContentBrowserDataSubsystem.cpp
--- a/Engine/Source/Editor/ContentBrowserData/Private/ContentBrowserDataSubsystem.cpp
+++ b/Engine/Source/Editor/ContentBrowserData/Private/ContentBrowserDataSubsystem.cpp
@@ -370,7 +370,7 @@ void UContentBrowserDataSubsystem::EnumerateItemsAtPath(const FName InPath, cons
 	}
 }
 
-bool UContentBrowserDataSubsystem::EnumerateItemsAtPaths(const TArrayView<class FContentBrowserItemPath> InItemPaths, const EContentBrowserItemTypeFilter InItemTypeFilter, TFunctionRef<bool(FContentBrowserItemData&&)> InCallback) const
+bool UContentBrowserDataSubsystem::EnumerateItemsAtPaths(const TArrayView<struct FContentBrowserItemPath> InItemPaths, const EContentBrowserItemTypeFilter InItemTypeFilter, TFunctionRef<bool(FContentBrowserItemData&&)> InCallback) const
 {
 	for (const auto& ActiveDataSourcePair : ActiveDataSources)
 	{
diff --git a/Engine/Source/Editor/ContentBrowserData/Public/ContentBrowserDataSubsystem.h b/Engine/Source/Editor/ContentBrowserData/Public/ContentBrowserDataSubsystem.h
--- a/Engine/Source/Editor/ContentBrowserData/Public/ContentBrowserDataSubsystem.h
+++ b/Engine/Source/Editor/ContentBrowserData/Public/ContentBrowserDataSubsystem.h
@@ -165,7 +165,7 @@ public:
 	 * @param InItemTypeFilter The types of items we want to find.
 	 * @param InCallback The function to invoke for each matching item (return true to continue enumeration).
 	 */
-	bool EnumerateItemsAtPaths(const TArrayView<class FContentBrowserItemPath> InItemPaths, const EContentBrowserItemTypeFilter InItemTypeFilter, TFunctionRef<bool(FContentBrowserItemData&&)> InCallback) const;
+	bool EnumerateItemsAtPaths(const TArrayView<struct FContentBrowserItemPath> InItemPaths, const EContentBrowserItemTypeFilter InItemTypeFilter, TFunctionRef<bool(FContentBrowserItemData&&)> InCallback) const;
 
 	/**
 	 * Enumerate the items (files) that exist for the given objects.
diff --git a/Engine/Source/Editor/ContentBrowserData/Public/ContentBrowserItemPath.h b/Engine/Source/Editor/ContentBrowserData/Public/ContentBrowserItemPath.h
--- a/Engine/Source/Editor/ContentBrowserData/Public/ContentBrowserItemPath.h
+++ b/Engine/Source/Editor/ContentBrowserData/Public/ContentBrowserItemPath.h
@@ -4,15 +4,20 @@
 
 #include "CoreMinimal.h"
 #include "ContentBrowserDataSubsystem.h"
+#include "Kismet/BlueprintFunctionLibrary.h"
+
+#include "ContentBrowserItemPath.generated.h"
 
 /**
  * Hold multiple versions of a path as FNames
  * 
  * Path conversion each time Set is called
  */
-class CONTENTBROWSERDATA_API FContentBrowserItemPath
+USTRUCT(BlueprintType, meta=(HasNativeBreak="ContentBrowserData.ContentBrowserItemPathExtensions.BreakContentBrowserItemPath", HasNativeMake="ContentBrowserData.ContentBrowserItemPathExtensions.MakeContentBrowserItemPath"))
+struct CONTENTBROWSERDATA_API FContentBrowserItemPath
 {
-public:
+	GENERATED_BODY()
+
 	FContentBrowserItemPath();
 
 	FContentBrowserItemPath(const FName InVirtualPath, const FName InInternalPath)
@@ -36,6 +41,10 @@ public:
 	 * Set the path being stored
 	 */
 	void SetPathFromString(const FStringView InPath, const EContentBrowserPathType InPathType);
+
+	/**
+	 * Set the path being stored
+	 */
 	void SetPathFromName(const FName InPath, const EContentBrowserPathType InPathType);
 
 	/**
@@ -70,3 +79,51 @@ private:
 	/** Path as internal (eg,. "/PluginA/MyFile") */
 	FName InternalPath;
 };
+
+UCLASS()
+class UContentBrowserItemPathExtensions : public UBlueprintFunctionLibrary
+{
+	GENERATED_BODY()
+
+public:
+
+	/**
+	 * Set the path being stored
+	 */
+	UFUNCTION(BlueprintCallable, Category = "Content Browser", meta = (ScriptMethod))
+	static void SetPath(UPARAM(ref) FContentBrowserItemPath& ItemPath, const FName InPath, const EContentBrowserPathType InPathType)
+	{
+		ItemPath.SetPathFromName(InPath, InPathType);
+	}
+
+	/**
+	 * Returns virtual path as FName (eg, "/All/Plugins/PluginA/MyFile").
+	 */
+	UFUNCTION(BlueprintCallable, Category = "Content Browser", meta = (ScriptMethod))
+	static FName GetVirtualPath(const FContentBrowserItemPath& ItemPath)
+	{
+		return ItemPath.GetVirtualPathName();
+	}
+
+	/**
+	 * Returns internal path if there is one (eg,. "/PluginA/MyFile").
+	 */
+	UFUNCTION(BlueprintCallable, Category = "Content Browser", meta = (ScriptMethod))
+	static FName GetInternalPath(const FContentBrowserItemPath& ItemPath)
+	{
+		return ItemPath.HasInternalPath() ? ItemPath.GetInternalPathName() : NAME_None;
+	}
+
+	UFUNCTION(BlueprintPure, Category = "Content Browser", meta = (Keywords = "construct build", NativeMakeFunc))
+	static FContentBrowserItemPath MakeContentBrowserItemPath(const FName InPath, const EContentBrowserPathType InPathType)
+	{
+		return FContentBrowserItemPath(InPath, InPathType);
+	}
+
+	UFUNCTION(BlueprintPure, Category = "Content Browser", meta = (NativeBreakFunc))
+	static void BreakContentBrowserItemPath(const FContentBrowserItemPath& ItemPath, FName& VirtualPath, FName& InternalPath)
+	{
+		VirtualPath = ItemPath.GetVirtualPathName();
+		InternalPath = ItemPath.HasInternalPath() ? ItemPath.GetInternalPathName() : NAME_None;
+	}
+ };
