Move implementation of FContentBrowserVirtualPathTree::GetMountPointFromPath to FPathViews

diff --git a/Engine/Source/Editor/ContentBrowser/Private/SPathView.cpp b/Engine/Source/Editor/ContentBrowser/Private/SPathView.cpp
--- a/Engine/Source/Editor/ContentBrowser/Private/SPathView.cpp
+++ b/Engine/Source/Editor/ContentBrowser/Private/SPathView.cpp
@@ -834,8 +834,7 @@ bool SPathView::InternalPathPassesBlockLists(const FStringView InInternalPath, c
 		const UContentBrowserSettings* ContentBrowserSettings = GetDefault<UContentBrowserSettings>();
 		if (ContentBrowserSettings->GetDisplayPluginFolders())
 		{
-			bool bHadClassesPrefix;
-			const FStringView FirstFolderName = FContentBrowserVirtualPathTree::GetMountPointFromPath(InInternalPath, bHadClassesPrefix);
+			const FStringView FirstFolderName = FPathViews::GetMountPointNameFromPath(InInternalPath);
 			if (TSharedPtr<IPlugin> Plugin = IPluginManager::Get().FindPlugin(FirstFolderName))
 			{
 				if (!PluginPathFilters->PassesAllFilters(Plugin.ToSharedRef()))
diff --git a/Engine/Source/Editor/ContentBrowserData/Private/ContentBrowserDataSubsystem.cpp b/Engine/Source/Editor/ContentBrowserData/Private/ContentBrowserDataSubsystem.cpp
--- a/Engine/Source/Editor/ContentBrowserData/Private/ContentBrowserDataSubsystem.cpp
+++ b/Engine/Source/Editor/ContentBrowserData/Private/ContentBrowserDataSubsystem.cpp
@@ -6,6 +6,7 @@
 #include "Containers/Ticker.h"
 #include "Misc/PackageName.h"
 #include "Misc/Paths.h"
+#include "Misc/PathViews.h"
 #include "Features/IModularFeatures.h"
 #include "Stats/Stats.h"
 #include "UObject/UObjectThreadContext.h"
@@ -744,8 +745,7 @@ void UContentBrowserDataSubsystem::ConvertInternalPathToVirtual(const FStringVie
 		}
 		else
 		{
-			bool bHadClassesPrefix = false;
-			const FStringView MountPointStringView = FContentBrowserVirtualPathTree::GetMountPointFromPath(InPath, bHadClassesPrefix);
+			const FStringView MountPointStringView = FPathViews::GetMountPointNameFromPath(InPath);
 
 			if (TSharedPtr<IPlugin> Plugin = IPluginManager::Get().FindPlugin(MountPointStringView))
 			{
diff --git a/Engine/Source/Editor/ContentBrowserData/Private/ContentBrowserVirtualPathTree.cpp b/Engine/Source/Editor/ContentBrowserData/Private/ContentBrowserVirtualPathTree.cpp
--- a/Engine/Source/Editor/ContentBrowserData/Private/ContentBrowserVirtualPathTree.cpp
+++ b/Engine/Source/Editor/ContentBrowserData/Private/ContentBrowserVirtualPathTree.cpp
@@ -3,13 +3,9 @@
 #include "ContentBrowserVirtualPathTree.h"
 #include "IContentBrowserDataModule.h"
 #include "ContentBrowserDataSubsystem.h"
-#include "Misc/StringBuilder.h"
 #include "Settings/ContentBrowserSettings.h"
-
-FContentBrowserVirtualPathTree::FContentBrowserVirtualPathTree()
-{
-	ClassesPrefix = TEXT("Classes_");
-}
+#include "Misc/StringBuilder.h"
+#include "Misc/PathViews.h"
 
 void FContentBrowserVirtualPathTree::Reset()
 {
@@ -20,30 +16,7 @@ void FContentBrowserVirtualPathTree::Reset()
 
 FStringView FContentBrowserVirtualPathTree::GetMountPointFromPath(const FStringView InPath, bool& bOutHadClassesPrefix)
 {
-	FStringView MountPointStringView;
-
-	int32 SecondForwardSlash = INDEX_NONE;
-	if (FStringView(InPath.GetData() + 1, InPath.Len() - 1).FindChar(TEXT('/'), SecondForwardSlash))
-	{
-		MountPointStringView = FStringView(InPath.GetData() + 1, SecondForwardSlash);
-	}
-	else
-	{
-		MountPointStringView = FStringView(InPath.GetData() + 1, InPath.Len() - 1);
-	}
-
-	static FString ClassesPrefix = TEXT("Classes_");
-	if (MountPointStringView.StartsWith(ClassesPrefix))
-	{
-		MountPointStringView.RightInline(MountPointStringView.Len() - ClassesPrefix.Len());
-		bOutHadClassesPrefix = true;
-	}
-	else
-	{
-		bOutHadClassesPrefix = false;
-	}
-
-	return MountPointStringView;
+	return FPathViews::GetMountPointNameFromPath(InPath, &bOutHadClassesPrefix);
 }
 
 bool FContentBrowserVirtualPathTree::CachePath(FName Path, FName InternalPath, TFunctionRef<void(FName)> OnPathAdded)
diff --git a/Engine/Source/Editor/ContentBrowserData/Public/ContentBrowserVirtualPathTree.h b/Engine/Source/Editor/ContentBrowserData/Public/ContentBrowserVirtualPathTree.h
--- a/Engine/Source/Editor/ContentBrowserData/Public/ContentBrowserVirtualPathTree.h
+++ b/Engine/Source/Editor/ContentBrowserData/Public/ContentBrowserVirtualPathTree.h
@@ -12,8 +12,6 @@
 class CONTENTBROWSERDATA_API FContentBrowserVirtualPathTree
 {
 public:
-	FContentBrowserVirtualPathTree();
-
 	/** Adds the specified path to the tree, creating nodes as needed and calling OnPathAdded for any new paths added. Returns true if the specified path was actually added (as opposed to already existed) */
 	bool CachePath(FName Path, FName InternalPath, TFunctionRef<void(FName)> OnPathAdded);
 
@@ -69,7 +67,8 @@ public:
 		return VirtualToInternalMounts;
 	}
 	
-	/** 
+	/**
+	 * @note Use FPathViews::GetMountPointNameFromPath instead
 	 * Returns name of the first folder in a path
 	 * Removes starting forward slash and Classes_ prefix 
 	 * Example: "/Classes_A/Textures" returns "A" and sets bOutHadClassesPrefix=true
@@ -94,9 +93,6 @@ private:
 	/** Get the string to use when show all folder is enabled */
 	const FString& GetAllFolderPrefix() const;
 
-	/** String of the classes prefix, by default it is "Classes_" */
-	FString ClassesPrefix;
-
 	/** A one-to-one mapping between a virtual path and internal path */
 	TMap<FName, FName> VirtualToInternalMounts;
 
diff --git a/Engine/Source/Runtime/Core/Private/Misc/PathViews.cpp b/Engine/Source/Runtime/Core/Private/Misc/PathViews.cpp
--- a/Engine/Source/Runtime/Core/Private/Misc/PathViews.cpp
+++ b/Engine/Source/Runtime/Core/Private/Misc/PathViews.cpp
@@ -675,4 +675,33 @@ void FPathViews::AppendPath(FStringBuilderBase& InOutPath, FStringView AppendPat
 		InOutPath.Reset();
 		InOutPath << AppendPath;
 	}
+}
+
+FStringView FPathViews::GetMountPointNameFromPath(const FStringView InPath, bool* bOutHadClassesPrefix /*= nullptr*/)
+{
+	FStringView MountPointStringView;
+
+	int32 SecondForwardSlash = INDEX_NONE;
+	if (FStringView(InPath.GetData() + 1, InPath.Len() - 1).FindChar(TEXT('/'), SecondForwardSlash))
+	{
+		MountPointStringView = FStringView(InPath.GetData() + 1, SecondForwardSlash);
+	}
+	else
+	{
+		MountPointStringView = FStringView(InPath.GetData() + 1, InPath.Len() - 1);
+	}
+
+	static const FString ClassesPrefix = TEXT("Classes_");
+	const bool bHasClassesPrefix = MountPointStringView.StartsWith(ClassesPrefix);
+	if (bHasClassesPrefix)
+	{
+		MountPointStringView.RightInline(MountPointStringView.Len() - ClassesPrefix.Len());
+	}
+
+	if (bOutHadClassesPrefix)
+	{
+		*bOutHadClassesPrefix = bHasClassesPrefix;
+	}
+
+	return MountPointStringView;
 }
\ No newline at end of file
diff --git a/Engine/Source/Runtime/Core/Public/Misc/PathViews.h b/Engine/Source/Runtime/Core/Public/Misc/PathViews.h
--- a/Engine/Source/Runtime/Core/Public/Misc/PathViews.h
+++ b/Engine/Source/Runtime/Core/Public/Misc/PathViews.h
@@ -337,4 +337,11 @@ public:
 	 * Does not interpret . or ..; each occurrence of these in either path will remain in the combined InOutPath.
 	 */
 	static void AppendPath(FStringBuilderBase& InOutPath, FStringView AppendPath);
+
+	/**
+	 * Returns the name of the first folder in a path
+	 * Removes starting forward slash and Classes_ prefix
+	 * Example: "/Classes_A/Textures" returns "A" and sets bOutHadClassesPrefix=true
+	 */
+	static FStringView GetMountPointNameFromPath(const FStringView InPath, bool* bOutHadClassesPrefix = nullptr);
 };
