Content Browser - Expose some tool menu context properties to scripting
#rnx

diff --git a/Engine/Source/Editor/ContentBrowserData/Public/ContentBrowserDataMenuContexts.h b/Engine/Source/Editor/ContentBrowserData/Public/ContentBrowserDataMenuContexts.h
--- a/Engine/Source/Editor/ContentBrowserData/Public/ContentBrowserDataMenuContexts.h
+++ b/Engine/Source/Editor/ContentBrowserData/Public/ContentBrowserDataMenuContexts.h
@@ -25,17 +25,17 @@ class CONTENTBROWSERDATA_API UContentBrowserDataMenuContext_AddNewMenu : public
 public:
 	DECLARE_DELEGATE_OneParam(FOnBeginItemCreation, const FContentBrowserItemDataTemporaryContext&);
 
-	UPROPERTY()
+	UPROPERTY(BlueprintReadOnly, Category = ContentBrowser)
 	TArray<FName> SelectedPaths;
 
 	// At least one of the selected paths maps to a mounted content root
-	UPROPERTY()
+	UPROPERTY(BlueprintReadOnly, Category = ContentBrowser)
 	bool bContainsValidPackagePath = false;
 
-	UPROPERTY()
+	UPROPERTY(BlueprintReadOnly, Category = ContentBrowser)
 	bool bCanBeModified = true;
 
-	UPROPERTY()
+	UPROPERTY(BlueprintReadOnly, Category = ContentBrowser)
 	EContentBrowserDataMenuContext_AddNewMenuDomain OwnerDomain = EContentBrowserDataMenuContext_AddNewMenuDomain::Toolbar;
 
 	FOnBeginItemCreation OnBeginItemCreation;
@@ -50,7 +50,7 @@ public:
 	UPROPERTY(BlueprintReadOnly, Category=ContentBrowser)
 	TArray<FContentBrowserItem> SelectedItems;
 
-	UPROPERTY()
+	UPROPERTY(BlueprintReadOnly, Category = ContentBrowser)
 	bool bCanBeModified = true;
 
 	TWeakPtr<SWidget> ParentWidget;
@@ -66,12 +66,12 @@ public:
 
 	DECLARE_DELEGATE(FOnRefreshView);
 
-	UPROPERTY()
+	UPROPERTY(BlueprintReadOnly, Category = ContentBrowser)
 	TArray<FContentBrowserItem> SelectedItems;
 
 	TArray<FCollectionNameType> SelectedCollections;
 
-	UPROPERTY()
+	UPROPERTY(BlueprintReadOnly, Category = ContentBrowser)
 	bool bCanBeModified = true;
 
 	TWeakPtr<SWidget> ParentWidget;
@@ -87,16 +87,16 @@ class CONTENTBROWSERDATA_API UContentBrowserDataMenuContext_DragDropMenu : publi
 	GENERATED_BODY()
 
 public:
-	UPROPERTY()
+	UPROPERTY(BlueprintReadOnly, Category = ContentBrowser)
 	FContentBrowserItem DropTargetItem;
 
-	UPROPERTY()
+	UPROPERTY(BlueprintReadOnly, Category = ContentBrowser)
 	TArray<FContentBrowserItem> DraggedItems;
 
-	UPROPERTY()
+	UPROPERTY(BlueprintReadOnly, Category = ContentBrowser)
 	bool bCanMove = true;
 
-	UPROPERTY()
+	UPROPERTY(BlueprintReadOnly, Category = ContentBrowser)
 	bool bCanCopy = true;
 
 	TWeakPtr<SWidget> ParentWidget;
