Prevent content browser ticks during inital content load

diff --git a/Engine/Source/Editor/ContentBrowserData/Private/ContentBrowserDataSubsystem.cpp b/Engine/Source/Editor/ContentBrowserData/Private/ContentBrowserDataSubsystem.cpp
--- a/Engine/Source/Editor/ContentBrowserData/Private/ContentBrowserDataSubsystem.cpp
+++ b/Engine/Source/Editor/ContentBrowserData/Private/ContentBrowserDataSubsystem.cpp
@@ -106,6 +106,8 @@ void UContentBrowserDataSubsystem::Initialize(FSubsystemCollectionBase& Collecti
 	FEditorDelegates::BeginPIE.AddUObject(this, &UContentBrowserDataSubsystem::OnBeginPIE);
 	FEditorDelegates::EndPIE.AddUObject(this, &UContentBrowserDataSubsystem::OnEndPIE);
 
+	FPackageName::OnContentPathMounted().AddUObject(this, &UContentBrowserDataSubsystem::OnContentPathMounted);
+
 	// Tick during normal operation
 	TickHandle = FTSTicker::GetCoreTicker().AddTicker(TEXT("ContentBrowserData"), 0.1f, [this](const float InDeltaTime)
 	{
@@ -129,6 +131,8 @@ void UContentBrowserDataSubsystem::Deinitialize()
 	FEditorDelegates::BeginPIE.RemoveAll(this);
 	FEditorDelegates::EndPIE.RemoveAll(this);
 
+	FPackageName::OnContentPathMounted().RemoveAll(this);
+
 	ActiveDataSources.Reset();
 	AvailableDataSources.Reset();
 	ActiveDataSourcesDiscoveringContent.Reset();
@@ -681,6 +685,13 @@ void UContentBrowserDataSubsystem::Tick(const float InDeltaTime)
 		return;
 	}
 
+	if (bContentMountedThisFrame)
+	{
+		// Content just added, defer tick for a frame or we risk slowing down content load
+		bContentMountedThisFrame = false;
+		return;
+	}
+
 	for (const auto& AvailableDataSourcePair : AvailableDataSources)
 	{
 		AvailableDataSourcePair.Value->Tick(InDeltaTime);
@@ -729,6 +740,11 @@ void UContentBrowserDataSubsystem::Tick(const float InDeltaTime)
 	}
 }
 
+void UContentBrowserDataSubsystem::OnContentPathMounted(const FString& AssetPath, const FString& ContentPath)
+{
+	bContentMountedThisFrame = true;
+}
+
 void UContentBrowserDataSubsystem::QueueItemDataUpdate(FContentBrowserItemDataUpdate&& InUpdate)
 {
 	// Ignore modified during PIE to reduce hitches, they will be updated when PIE stops
diff --git a/Engine/Source/Editor/ContentBrowserData/Public/ContentBrowserDataSubsystem.h b/Engine/Source/Editor/ContentBrowserData/Public/ContentBrowserDataSubsystem.h
--- a/Engine/Source/Editor/ContentBrowserData/Public/ContentBrowserDataSubsystem.h
+++ b/Engine/Source/Editor/ContentBrowserData/Public/ContentBrowserDataSubsystem.h
@@ -376,6 +376,12 @@ private:
 	 */
 	void OnEndPIE(const bool bIsSimulating);
 
+	/**
+	 * Called when Content added to delay content browser tick for a frame
+	 * Prevents content browser from slowing down initialization by ticking as content is loaded
+	 */
+	void OnContentPathMounted(const FString& AssetPath, const FString& ContentPath);
+
 	//~ IContentBrowserItemDataSink interface
 	virtual void QueueItemDataUpdate(FContentBrowserItemDataUpdate&& InUpdate) override;
 	virtual void NotifyItemDataRefreshed() override;
@@ -428,6 +434,11 @@ private:
 	 */
 	bool bIsPIEActive = false;
 
+	/**
+	 * True if content was just mounted this frame
+	 */
+	bool bContentMountedThisFrame = false;
+
 	/**
 	 * >0 if Tick events have currently been suppressed.
 	 * @see FScopedSuppressContentBrowserDataTick
