Fixed nested virtual roots (from aliases) failing to resolve correctly

It's possible for aliases to generate nested virtual roots, eg)
 /All/MyStuff -> /MyStuff
 /All/MyStuff/SomePlugin -> /SomePlugin

In that situation the virtual path resolver would fail to process a path like "/All/MyStuff/RefContent" since "/All/MyStuff" had a child folder so was assumed never to be a virtual root in its own right.

This change updates the resolver to keep track of the best potential virtual root as the path is parsed (even if it has children), and will then use the best match when converting to an internal path.

diff --git a/Engine/Source/Editor/ContentBrowserData/Private/ContentBrowserVirtualPathTree.cpp b/Engine/Source/Editor/ContentBrowserData/Private/ContentBrowserVirtualPathTree.cpp
--- a/Engine/Source/Editor/ContentBrowserData/Private/ContentBrowserVirtualPathTree.cpp
+++ b/Engine/Source/Editor/ContentBrowserData/Private/ContentBrowserVirtualPathTree.cpp
@@ -355,8 +355,8 @@ EContentBrowserPathType FContentBrowserVirtualPathTree::TryConvertVirtualPathToI
 	}
 
 	// Walk path until reaches internal mount point or ends remaining as a virtual path
-	bool bReachedMountPoint = false;
 	bool bFailedPathExists = false;
+	int32 BestMountPointSplitIndex = INDEX_NONE;
 	int32 SplitIndex = 0;
 	const TCHAR* PathStr = InPath.GetData();
 	const TCHAR* PathStrEnd = PathStr + InPath.Len();
@@ -371,11 +371,17 @@ EContentBrowserPathType FContentBrowserVirtualPathTree::TryConvertVirtualPathToI
 			{
 				if (!bCheckPathIsFullyVirtual)
 				{
-					bReachedMountPoint = true;
+					BestMountPointSplitIndex = SplitIndex;
 					break;
 				}
 				else
 				{
+					if (VirtualToInternalMounts.Contains(CheckPath))
+					{
+						// This virtual path has other virtual paths as children (which is why bCheckPathIsFullyVirtual returned true)
+						// It's still a valid root to map against though, so keep track of it until we find a better match to use
+						BestMountPointSplitIndex = SplitIndex;
+					}
 					SplitIndex = (int32)(PathCharPtr - PathStr);
 				}
 			}
@@ -397,7 +403,7 @@ EContentBrowserPathType FContentBrowserVirtualPathTree::TryConvertVirtualPathToI
 
 	if (SplitIndex > 0)
 	{
-		if (!bReachedMountPoint)
+		if (BestMountPointSplitIndex == INDEX_NONE)
 		{
 			if (bFailedPathExists)
 			{
@@ -427,17 +433,14 @@ EContentBrowserPathType FContentBrowserVirtualPathTree::TryConvertVirtualPathToI
 		}
 
 		FStringView InternalPathView(InPath);
-		InternalPathView.RightChopInline(SplitIndex);
+		InternalPathView.RightChopInline(BestMountPointSplitIndex);
 		OutPath.Append(InternalPathView);
 		return EContentBrowserPathType::Internal;
 	}
-	else
+	else if (BestMountPointSplitIndex != INDEX_NONE)
 	{
-		if (bReachedMountPoint)
-		{
-			OutPath.Append(InPath);
-			return EContentBrowserPathType::Internal;
-		}
+		OutPath.Append(InPath);
+		return EContentBrowserPathType::Internal;
 	}
 
 	return EContentBrowserPathType::None;
