Return true earlier from content browser folder functions that search multiple data sources, this improves performance while refreshing things

diff --git a/Engine/Source/Editor/ContentBrowserData/Private/ContentBrowserDataSubsystem.cpp b/Engine/Source/Editor/ContentBrowserData/Private/ContentBrowserDataSubsystem.cpp
--- a/Engine/Source/Editor/ContentBrowserData/Private/ContentBrowserDataSubsystem.cpp
+++ b/Engine/Source/Editor/ContentBrowserData/Private/ContentBrowserDataSubsystem.cpp
@@ -539,7 +539,6 @@ bool UContentBrowserDataSubsystem::PrioritizeSearchPath(const FName InPath)
 
 bool UContentBrowserDataSubsystem::IsFolderVisible(const FName InPath, const EContentBrowserIsFolderVisibleFlags InFlags) const
 {
-	bool bIsVisible = false;
 	bool bIsKnownPath = false;
 	for (const auto& ActiveDataSourcePair : ActiveDataSources)
 	{
@@ -547,10 +546,15 @@ bool UContentBrowserDataSubsystem::IsFolderVisible(const FName InPath, const ECo
 		if (DataSource->IsVirtualPathUnderMountRoot(InPath))
 		{
 			bIsKnownPath = true;
-			bIsVisible |= DataSource->IsFolderVisible(InPath, InFlags);
+			if (DataSource->IsFolderVisible(InPath, InFlags))
+			{
+				return true;
+			}
 		}
 	}
-	return bIsVisible || !bIsKnownPath;
+
+	// Return true if this is visible for any sources, or this path isn't handled by any of the sources
+	return !bIsKnownPath;
 }
 
 bool UContentBrowserDataSubsystem::IsFolderVisibleIfHidingEmpty(const FName InPath) const
@@ -560,16 +564,18 @@ bool UContentBrowserDataSubsystem::IsFolderVisibleIfHidingEmpty(const FName InPa
 
 bool UContentBrowserDataSubsystem::CanCreateFolder(const FName InPath, FText* OutErrorMsg) const
 {
-	bool bCanCreateFolder = false;
 	for (const auto& ActiveDataSourcePair : ActiveDataSources)
 	{
 		UContentBrowserDataSource* DataSource = ActiveDataSourcePair.Value;
 		if (DataSource->IsVirtualPathUnderMountRoot(InPath))
 		{
-			bCanCreateFolder |= DataSource->CanCreateFolder(InPath, OutErrorMsg);
+			if (DataSource->CanCreateFolder(InPath, OutErrorMsg))
+			{
+				return true;
+			}
 		}
 	}
-	return bCanCreateFolder;
+	return false;
 }
 
 FContentBrowserItemTemporaryContext UContentBrowserDataSubsystem::CreateFolder(const FName InPath) const
