Add IsInPlugin to FContentBrowserItem.

diff --git a/Engine/Source/Editor/ContentBrowserData/Private/ContentBrowserItem.cpp b/Engine/Source/Editor/ContentBrowserData/Private/ContentBrowserItem.cpp
--- a/Engine/Source/Editor/ContentBrowserData/Private/ContentBrowserItem.cpp
+++ b/Engine/Source/Editor/ContentBrowserData/Private/ContentBrowserItem.cpp
@@ -290,6 +290,12 @@ bool FContentBrowserItem::IsFile() const
 	return PrimaryItemData && PrimaryItemData->IsFile();
 }
 
+bool FContentBrowserItem::IsInPlugin() const
+{
+	const FContentBrowserItemData* PrimaryItemData = GetPrimaryInternalItem();
+	return PrimaryItemData && PrimaryItemData->IsPlugin();
+}
+
 bool FContentBrowserItem::IsSupported() const
 {
 	const FContentBrowserItemData* PrimaryItemData = GetPrimaryInternalItem();
diff --git a/Engine/Source/Editor/ContentBrowserData/Public/ContentBrowserItem.h b/Engine/Source/Editor/ContentBrowserData/Public/ContentBrowserItem.h
--- a/Engine/Source/Editor/ContentBrowserData/Public/ContentBrowserItem.h
+++ b/Engine/Source/Editor/ContentBrowserData/Public/ContentBrowserItem.h
@@ -135,6 +135,12 @@ public:
 	 */
 	bool IsFile() const;
 
+	/**
+	 * Check to see whether this item is in a plugin.
+	 * @note Equivalent to testing whether EContentBrowserItemFlags::Category_Plugin is set on GetItemFlags().
+	 */
+	bool IsInPlugin() const;
+
 	/**
 	 * Check if the item is representing a supported item
 	 * The content browser can also display some unsupported asset
