Optimising the time it take to exit a pie session in large project by removing need to refresh of the views that use the data of the ContentBrowserDataSubsystem.

The skipped event are now instead delayed to be proccessed at the end of the pie session.

diff --git a/Engine/Source/Editor/ContentBrowserData/Private/ContentBrowserDataSubsystem.cpp b/Engine/Source/Editor/ContentBrowserData/Private/ContentBrowserDataSubsystem.cpp
--- a/Engine/Source/Editor/ContentBrowserData/Private/ContentBrowserDataSubsystem.cpp
+++ b/Engine/Source/Editor/ContentBrowserData/Private/ContentBrowserDataSubsystem.cpp
@@ -760,7 +760,7 @@ void UContentBrowserDataSubsystem::Tick(const float InDeltaTime)
 		TRACE_CPUPROFILER_EVENT_SCOPE(UContentBrowserDataSubsystem::BroadcastItemDataRefreshed);
 
 		bPendingItemDataRefreshedNotification = false;
-		bHasIgnoredItemUpdates = false;
+		DelayedPendingUpdates.Empty();
 		PendingUpdates.Empty();
 		ItemDataRefreshedDelegate.Broadcast();
 	}
@@ -809,10 +809,34 @@ void UContentBrowserDataSubsystem::OnContentPathMounted(const FString& AssetPath
 
 void UContentBrowserDataSubsystem::QueueItemDataUpdate(FContentBrowserItemDataUpdate&& InUpdate)
 {
-	// Ignore modified during PIE to reduce hitches, they will be updated when PIE stops
-	if ((InUpdate.GetUpdateType() == EContentBrowserItemUpdateType::Modified) && !AllowModifiedItemDataUpdates())
+	if (!AllowModifiedItemDataUpdates())
 	{
-		bHasIgnoredItemUpdates = true;
+		const EContentBrowserItemUpdateType UpdateType = InUpdate.GetUpdateType();
+
+		// Ignore modified during PIE to reduce hitches, they will be queue and then added to the pending updates when PIE stops
+		if (UpdateType == EContentBrowserItemUpdateType::Modified)
+		{
+			FContentBrowserItemKey ItemKey(InUpdate.GetItemData());
+			DelayedPendingUpdates.Add(MoveTemp(ItemKey), MoveTemp(InUpdate));
+		}
+		else
+		{
+			// Clear the delayed update for the item if there was one 
+			if (UpdateType == EContentBrowserItemUpdateType::Moved)
+			{
+				const FContentBrowserItemData& ItemData = InUpdate.GetItemData();
+				FContentBrowserItemKey ItemKey(ItemData.GetItemType(), InUpdate.GetPreviousVirtualPath(), ItemData.GetOwnerDataSource());
+				DelayedPendingUpdates.Remove(ItemKey);
+			}
+			else
+			{
+				FContentBrowserItemKey ItemKey(InUpdate.GetItemData());
+				DelayedPendingUpdates.Remove(ItemKey);
+			}
+
+			// TODO: Merge multiple Modified updates for a single item?
+			PendingUpdates.Emplace(MoveTemp(InUpdate));
+		}
 	}
 	else
 	{
@@ -840,10 +864,17 @@ void UContentBrowserDataSubsystem::OnEndPIE(const bool bIsSimulating)
 {
 	bIsPIEActive = false;
 
-	if (bHasIgnoredItemUpdates)
+	if (!DelayedPendingUpdates.IsEmpty())
 	{
-		// Perform a full update because modified updates were ignored during PIE
-		NotifyItemDataRefreshed();
+		// Move the DelayedPendingUpdates into the PendingUpdates.
+		PendingUpdates.Reserve(DelayedPendingUpdates.Num() + PendingUpdates.Num());
+		
+		for (TPair<FContentBrowserItemKey, FContentBrowserItemDataUpdate>& Pair : DelayedPendingUpdates)
+		{
+			PendingUpdates.Add(MoveTemp(Pair.Value));
+		}
+
+		DelayedPendingUpdates.Empty();
 	}
 }
 
diff --git a/Engine/Source/Editor/ContentBrowserData/Public/ContentBrowserDataSubsystem.h b/Engine/Source/Editor/ContentBrowserData/Public/ContentBrowserDataSubsystem.h
--- a/Engine/Source/Editor/ContentBrowserData/Public/ContentBrowserDataSubsystem.h
+++ b/Engine/Source/Editor/ContentBrowserData/Public/ContentBrowserDataSubsystem.h
@@ -467,14 +467,15 @@ private:
 	TArray<FContentBrowserItemDataUpdate> PendingUpdates;
 
 	/**
-	 * True if an item data refresh notification is pending.
+	 * Set of item data updates that are delayed to preserve the editor performance.
+	 * These will be passed to the PendingUpdates when we exit a pie session.
 	 */
-	bool bPendingItemDataRefreshedNotification = false;
+	TMap<FContentBrowserItemKey, FContentBrowserItemDataUpdate> DelayedPendingUpdates;
 
 	/**
-	 * True if there are currently any ignored changes.
+	 * True if an item data refresh notification is pending.
 	 */
-	bool bHasIgnoredItemUpdates = false;
+	bool bPendingItemDataRefreshedNotification = false;
 
 	/**
 	 * True if Play in Editor is active.
