Remove the need for a data refresh on the fist tick of the editor.

Deprecate the rest of the public api to request a full data refresh from the ContentBrowserDataSubsystem.

diff --git a/Engine/Source/Editor/ContentBrowserData/Private/ContentBrowserDataSource.cpp b/Engine/Source/Editor/ContentBrowserData/Private/ContentBrowserDataSource.cpp
--- a/Engine/Source/Editor/ContentBrowserData/Private/ContentBrowserDataSource.cpp
+++ b/Engine/Source/Editor/ContentBrowserData/Private/ContentBrowserDataSource.cpp
@@ -556,7 +556,9 @@ void UContentBrowserDataSource::NotifyItemDataRefreshed()
 {
 	if (DataSink)
 	{
+PRAGMA_DISABLE_DEPRECATION_WARNINGS
 		DataSink->NotifyItemDataRefreshed();
+PRAGMA_ENABLE_DEPRECATION_WARNINGS
 	}
 }
 
diff --git a/Engine/Source/Editor/ContentBrowserData/Private/ContentBrowserDataSubsystem.cpp b/Engine/Source/Editor/ContentBrowserData/Private/ContentBrowserDataSubsystem.cpp
--- a/Engine/Source/Editor/ContentBrowserData/Private/ContentBrowserDataSubsystem.cpp
+++ b/Engine/Source/Editor/ContentBrowserData/Private/ContentBrowserDataSubsystem.cpp
@@ -98,6 +98,13 @@ void UContentBrowserDataSubsystem::Initialize(FSubsystemCollectionBase& Collecti
 		{
 			HandleDataSourceRegistered(DataSourceFeatureName, ModularFeatures.GetModularFeatureImplementation(DataSourceFeatureName, AvailableDataSourcesIndex));
 		}
+
+		/**
+		 * If any view already exist refresh them now instead of waiting.
+		 * This avoid asking for the view that where just created to refresh their data next frame during the editor initialization.
+		 */ 
+		bPendingItemDataRefreshedNotification = false;
+		ItemDataRefreshedDelegate.Broadcast();
 	}
 
 	ModularFeatures.OnModularFeatureRegistered().AddUObject(this, &UContentBrowserDataSubsystem::HandleDataSourceRegistered);
@@ -160,7 +167,9 @@ bool UContentBrowserDataSubsystem::ActivateDataSource(const FName Name)
 			DataSource->SetDataSink(this);
 			ActiveDataSources.Add(Name, DataSource);
 			ActiveDataSourcesDiscoveringContent.Add(Name);
+PRAGMA_DISABLE_DEPRECATION_WARNINGS
 			NotifyItemDataRefreshed();
+PRAGMA_ENABLE_DEPRECATION_WARNINGS
 			return true;
 		}
 		else
@@ -181,7 +190,9 @@ bool UContentBrowserDataSubsystem::DeactivateDataSource(const FName Name)
 		DataSource->SetDataSink(nullptr);
 		ActiveDataSources.Remove(Name);
 		ActiveDataSourcesDiscoveringContent.Remove(Name);
+PRAGMA_DISABLE_DEPRECATION_WARNINGS
 		NotifyItemDataRefreshed();
+PRAGMA_ENABLE_DEPRECATION_WARNINGS
 		return true;
 	}
 
@@ -205,7 +216,9 @@ void UContentBrowserDataSubsystem::ActivateAllDataSources()
 		// Merge this array as it may contain sources that we've not yet discovered, so can't activate yet
 		EnabledDataSources.AddUnique(ActiveDataSourcePair.Key);
 	}
+PRAGMA_DISABLE_DEPRECATION_WARNINGS
 	NotifyItemDataRefreshed();
+PRAGMA_ENABLE_DEPRECATION_WARNINGS
 }
 
 void UContentBrowserDataSubsystem::DeactivateAllDataSources()
@@ -223,7 +236,10 @@ void UContentBrowserDataSubsystem::DeactivateAllDataSources()
 	ActiveDataSources.Reset();
 	EnabledDataSources.Reset();
 	ActiveDataSourcesDiscoveringContent.Reset();
+
+PRAGMA_DISABLE_DEPRECATION_WARNINGS
 	NotifyItemDataRefreshed();
+PRAGMA_ENABLE_DEPRECATION_WARNINGS
 }
 
 TArray<FName> UContentBrowserDataSubsystem::GetAvailableDataSources() const
diff --git a/Engine/Source/Editor/ContentBrowserData/Public/ContentBrowserDataSubsystem.h b/Engine/Source/Editor/ContentBrowserData/Public/ContentBrowserDataSubsystem.h
--- a/Engine/Source/Editor/ContentBrowserData/Public/ContentBrowserDataSubsystem.h
+++ b/Engine/Source/Editor/ContentBrowserData/Public/ContentBrowserDataSubsystem.h
@@ -431,6 +431,8 @@ private:
 
 	//~ IContentBrowserItemDataSink interface
 	virtual void QueueItemDataUpdate(FContentBrowserItemDataUpdate&& InUpdate) override;
+
+	UE_DEPRECATED(5.5, "NotifyItemDataRefreshed is deprecated, for editor performance reasons no external systems should be able to request a full refresh.")
 	virtual void NotifyItemDataRefreshed() override;
 
 	/**
diff --git a/Engine/Source/Editor/ContentBrowserData/Public/ContentBrowserItemData.h b/Engine/Source/Editor/ContentBrowserData/Public/ContentBrowserItemData.h
--- a/Engine/Source/Editor/ContentBrowserData/Public/ContentBrowserItemData.h
+++ b/Engine/Source/Editor/ContentBrowserData/Public/ContentBrowserItemData.h
@@ -743,6 +743,7 @@ public:
 	/**
 	 * Notify a wholesale item data update, for data sources that can't provide delta-updates.
 	 */
+	 UE_DEPRECATED(5.5, "NotifyItemDataRefreshed is deprecated, for editor performance reasons no external systems should be able to request a full refresh.")
 	virtual void NotifyItemDataRefreshed() = 0;
 
 	/**
