Add copyright header.

diff --git a/Engine/Source/Editor/ContentBrowserData/Public/GetOrEnumerateSink.h b/Engine/Source/Editor/ContentBrowserData/Public/GetOrEnumerateSink.h
--- a/Engine/Source/Editor/ContentBrowserData/Public/GetOrEnumerateSink.h
+++ b/Engine/Source/Editor/ContentBrowserData/Public/GetOrEnumerateSink.h
@@ -1,3 +1,5 @@
+// Copyright Epic Games, Inc. All Rights Reserved.
+
 #pragma once
 
 #include "Containers/Array.h"
