Fix incorrectly implemented comparison operators.

diff --git a/Engine/Source/Editor/ContentBrowserData/Public/ContentBrowserItemData.h b/Engine/Source/Editor/ContentBrowserData/Public/ContentBrowserItemData.h
--- a/Engine/Source/Editor/ContentBrowserData/Public/ContentBrowserItemData.h
+++ b/Engine/Source/Editor/ContentBrowserData/Public/ContentBrowserItemData.h
@@ -642,7 +642,7 @@ public:
 	}
 	bool operator!=(const FContentBrowserMinimalItemData& Other) const
 	{
-		return ItemType == Other.ItemType && VirtualPath == Other.VirtualPath && DataSource == Other.DataSource;
+		return ItemType != Other.ItemType || VirtualPath != Other.VirtualPath || DataSource != Other.DataSource;
 	}
 
 	/** Return a string representation for debugging. */
