Added a display name for the virtual "All" root folder

#rnx

diff --git a/Engine/Source/Editor/ContentBrowserData/Private/ContentBrowserDataSource.cpp b/Engine/Source/Editor/ContentBrowserData/Private/ContentBrowserDataSource.cpp
--- a/Engine/Source/Editor/ContentBrowserData/Private/ContentBrowserDataSource.cpp
+++ b/Engine/Source/Editor/ContentBrowserData/Private/ContentBrowserDataSource.cpp
@@ -587,8 +587,14 @@ FContentBrowserItemData UContentBrowserDataSource::CreateVirtualFolderItem(const
 {
 	const FString FolderItemName = FPackageName::GetShortName(InFolderPath);
 
+	static const FName AllRootPath = "/All";
+
 	FText FolderDisplayNameOverride;
-	if (FolderItemName == TEXT("GameData"))
+	if (InFolderPath == AllRootPath)
+	{
+		FolderDisplayNameOverride = LOCTEXT("AllFolderDisplayName", "All");
+	}
+	else if (FolderItemName == TEXT("GameData"))
 	{
 		FolderDisplayNameOverride = LOCTEXT("GameDataFolderDisplayName", "Game Data");
 	}
