[Content Browser]
- Fix GetInvariantPath() that was returning the virtual path with /All folder prefix
- Delegate to use the Verse path of a plugin as its virtual path in the content browser
#rnx

diff --git a/Engine/Source/Editor/ContentBrowserData/Private/ContentBrowserDataSubsystem.cpp b/Engine/Source/Editor/ContentBrowserData/Private/ContentBrowserDataSubsystem.cpp
--- a/Engine/Source/Editor/ContentBrowserData/Private/ContentBrowserDataSubsystem.cpp
+++ b/Engine/Source/Editor/ContentBrowserData/Private/ContentBrowserDataSubsystem.cpp
@@ -24,6 +24,7 @@
 #include "Settings/ContentBrowserSettings.h"
 #include "Stats/Stats.h"
 #include "Stats/Stats2.h"
+#include "String/RemoveFrom.h"
 #include "Templates/Function.h"
 #include "Templates/Less.h"
 #include "Templates/SharedPointer.h"
@@ -933,6 +934,53 @@ void UContentBrowserDataSubsystem::ConvertInternalPathToVirtual(const FStringVie
 		}
 	}
 
+	TOptional<FStringView> MountPointOptional;
+	auto GetMountPoint = [&MountPointOptional, InPath]()
+	{
+		if (!MountPointOptional.IsSet())
+		{
+			MountPointOptional = FPathViews::GetMountPointNameFromPath(InPath);
+		}
+		return MountPointOptional.GetValue();
+	};
+
+	TOptional<TSharedPtr<IPlugin>> PluginOptional;
+	auto GetPlugin = [&PluginOptional, GetMountPoint]()
+	{
+		if (!PluginOptional.IsSet())
+		{
+			PluginOptional = IPluginManager::Get().FindPlugin(GetMountPoint());
+		}
+		return PluginOptional.GetValue();
+	};
+
+	if (UsePluginVersePathDelegate.IsBound())
+	{
+		if (const TSharedPtr<IPlugin> Plugin = GetPlugin())
+		{
+			if (UsePluginVersePath(Plugin.ToSharedRef()))
+			{
+			#if 0
+				// @fixme The semantically correct thing to do would be this:
+				OutPath.Append(FPathViews::GetPath(UE::String::RemoveFromEnd(FStringView(Plugin->GetVersePath()), TEXTVIEW("/"))));
+			#else
+				// However we have an issue with multi-plugin projects to solve first
+				// The root module/plugin uses the project Verse path
+				// Non-root modules are "faking" their Verse path this way: [REDACTED_EMAIL]/project/module
+				// It's semantically invalid in Verse because that module isn't actually a sub-module of the root module
+				// This creates problems in the content browser: the root plugin being an actual folder, the virtual path 
+				// of other plugins cannot be start with it (i.e. we can't mix real folder hierarchies with virtual path hierarchies)
+				// Another reason to fix this is to support "namespaces" in a project Verse paths such as: 
+				// /owner@domain/purely/organizational/structure/project
+				OutPath.AppendChar(TCHAR('/'));
+				OutPath.Append(FPathViews::GetMountPointNameFromPath(FStringView(Plugin->GetVersePath())));
+			#endif
+				OutPath.Append(InPath.GetData(), InPath.Len());
+				return;
+			}
+		}
+	}
+
 	if (GetDefault<UContentBrowserSettings>()->bOrganizeFolders && InPath.Len() > 1)
 	{
 		if (GenerateVirtualPathPrefixDelegate.IsBound())
@@ -941,9 +989,7 @@ void UContentBrowserDataSubsystem::ConvertInternalPathToVirtual(const FStringVie
 		}
 		else
 		{
-			const FStringView MountPointStringView = FPathViews::GetMountPointNameFromPath(InPath);
-
-			if (TSharedPtr<IPlugin> Plugin = IPluginManager::Get().FindPlugin(MountPointStringView))
+			if (TSharedPtr<IPlugin> Plugin = GetPlugin())
 			{
 				if (Plugin->GetLoadedFrom() == EPluginLoadedFrom::Engine)
 				{
@@ -974,7 +1020,7 @@ void UContentBrowserDataSubsystem::ConvertInternalPathToVirtual(const FStringVie
 					}
 				}
 			}
-			else if (MountPointStringView.Equals(TEXT("Engine")))
+			else if (GetMountPoint().Equals(TEXT("Engine")))
 			{
 				OutPath.Append(TEXT("/EngineData"));
 			}
@@ -1030,6 +1076,27 @@ FContentBrowserGenerateVirtualPathDelegate& UContentBrowserDataSubsystem::OnGene
 	return GenerateVirtualPathPrefixDelegate;
 }
 
+bool UContentBrowserDataSubsystem::UsePluginVersePath(const TSharedRef<IPlugin>& Plugin)
+{
+	if (UsePluginVersePathDelegate.IsBound() && !Plugin->GetVersePath().IsEmpty())
+	{
+		return UsePluginVersePathDelegate.Execute(Plugin);
+	}
+	return false;
+}
+
+void UContentBrowserDataSubsystem::SetUsePluginVersePathDelegate(FContentBrowserUsePluginVersePathDelegate InDelegate)
+{
+	UsePluginVersePathDelegate = MoveTemp(InDelegate);
+	SetVirtualPathTreeNeedsRebuild();
+	RefreshVirtualPathTreeIfNeeded();
+}
+
+FContentBrowserUsePluginVersePathDelegate& UContentBrowserDataSubsystem::GetUsePluginVersePathDelegate()
+{
+	return UsePluginVersePathDelegate;
+}
+
 const FString& UContentBrowserDataSubsystem::GetAllFolderPrefix() const
 {
 	return AllFolderPrefix;
diff --git a/Engine/Source/Editor/ContentBrowserData/Private/ContentBrowserDataUtils.cpp b/Engine/Source/Editor/ContentBrowserData/Private/ContentBrowserDataUtils.cpp
--- a/Engine/Source/Editor/ContentBrowserData/Private/ContentBrowserDataUtils.cpp
+++ b/Engine/Source/Editor/ContentBrowserData/Private/ContentBrowserDataUtils.cpp
@@ -4,6 +4,9 @@
 
 #include "AssetRegistry/AssetData.h"
 #include "Containers/StringView.h"
+#include "ContentBrowserDataSubsystem.h"
+#include "Editor.h"
+#include "IContentBrowserDataModule.h"
 #include "Interfaces/IPluginManager.h"
 #include "Internationalization/Internationalization.h"
 #include "Math/UnrealMathSSE.h"
@@ -201,6 +204,11 @@ bool ContentBrowserDataUtils::PathPassesAttributeFilter(const FStringView InPath
 
 FText ContentBrowserDataUtils::GetFolderItemDisplayNameOverride(const FName InFolderPath, const FString& InFolderItemName, const bool bIsClassesFolder, const bool bIsCookedPath)
 {
+	if (InFolderItemName.IsEmpty())
+	{
+		return FText::GetEmpty();
+	}
+
 	FText FolderDisplayNameOverride;
 
 	if (!bIsClassesFolder)
@@ -239,39 +247,68 @@ FText ContentBrowserDataUtils::GetFolderItemDisplayNameOverride(const FName InFo
 					TopLevelFolderName.RightChopInline(ClassesPrefix.Len());
 				}
 			}
-
-			TSharedPtr<IPlugin> Plugin = IPluginManager::Get().FindPlugin(TopLevelFolderName);
-
-			FString OverrideName;
-			if (GetDefault<UContentBrowserSettings>()->bDisplayFriendlyNameForPluginFolders && Plugin && Plugin->GetFriendlyName().Len() > 0)
-			{
-				OverrideName = Plugin->GetFriendlyName();
-			}
-			if (OverrideName.IsEmpty())
+			
+			// If we're using the Verse path of the plugin in the content browser, don't modify it's folder name
+			const TSharedPtr<IPlugin> Plugin = IPluginManager::Get().FindPlugin(TopLevelFolderName);
+			if (Plugin)
 			{
-				OverrideName = TopLevelFolderName;
+				UContentBrowserDataSubsystem* ContentBrowser = IContentBrowserDataModule::Get().GetSubsystem();
+				if (ContentBrowser && ContentBrowser->UsePluginVersePath(Plugin.ToSharedRef()))
+				{
+				#if 0
+					// @fixme The semantically correct thing to do would be this:
+					FolderDisplayNameOverride = FText::FromString(FPaths::GetPathLeaf(Plugin->GetVersePath()));
+				#else
+					// However we have an issue with multi-plugin projects to solve first
+					// The root module/plugin uses the project Verse path
+					// Non-root modules are "faking" their Verse path this way: [REDACTED_EMAIL]/project/module
+					// It's semantically invalid in Verse because that module isn't actually a sub-module of the root module
+					// This creates problems in the content browser: the root plugin being an actual folder, the virtual path 
+					// of other plugins cannot be start with it (i.e. we can't mix real folder hierarchies with virtual path hierarchies)
+					// Another reason to fix this is to support "namespaces" in a project Verse paths such as: 
+					// /owner@domain/purely/organizational/structure/project
+					const FStringView VersePath = Plugin->GetVersePath();
+					const int32 VersePathDomainLen = FPathViews::GetMountPointNameFromPath(VersePath, /*bOutHadClassesPrefix=*/nullptr, /*bInWithoutSlashes=*/false).Len() + 1;
+					FString VersePathRemainder(VersePath.RightChop(VersePathDomainLen));
+					VersePathRemainder.ReplaceInline(TEXT("/"), TEXT(" / "));
+					FolderDisplayNameOverride = FText::FromString(MoveTemp(VersePathRemainder));
+				#endif
+				}
 			}
 
-			if (bIsClassesFolder)
-			{
-				FolderDisplayNameOverride = FText::Format(LOCTEXT("ClassFolderDisplayNameFmt", "{0} C++ Classes"), FText::AsCultureInvariant(OverrideName));
-			}
-			else
+			if (FolderDisplayNameOverride.IsEmpty())
 			{
-				bool bDisplayContentFolderSuffix = GetDefault<UContentBrowserSettings>()->bDisplayContentFolderSuffix;
-				if (bDisplayContentFolderSuffix && Plugin && Plugin->GetDescriptor().Modules.Num() == 0 && bIsCookedPath)
+				FString OverrideName;
+				if (GetDefault<UContentBrowserSettings>()->bDisplayFriendlyNameForPluginFolders && Plugin && Plugin->GetFriendlyName().Len() > 0)
+				{
+					OverrideName = Plugin->GetFriendlyName();
+				}
+				if (OverrideName.IsEmpty())
 				{
-					// Exclude the content suffix for plugins that only contain cooked content and have no C++ modules
-					bDisplayContentFolderSuffix = false;
+					OverrideName = TopLevelFolderName;
 				}
 
-				if (bDisplayContentFolderSuffix)
+				if (bIsClassesFolder)
 				{
-					FolderDisplayNameOverride = FText::Format(LOCTEXT("ContentFolderDisplayNameFmt", "{0} Content"), FText::AsCultureInvariant(OverrideName));
+					FolderDisplayNameOverride = FText::Format(LOCTEXT("ClassFolderDisplayNameFmt", "{0} C++ Classes"), FText::AsCultureInvariant(OverrideName));
 				}
 				else
 				{
-					FolderDisplayNameOverride = FText::AsCultureInvariant(OverrideName);
+					bool bDisplayContentFolderSuffix = GetDefault<UContentBrowserSettings>()->bDisplayContentFolderSuffix;
+					if (bDisplayContentFolderSuffix && Plugin && Plugin->GetDescriptor().Modules.Num() == 0 && bIsCookedPath)
+					{
+						// Exclude the content suffix for plugins that only contain cooked content and have no C++ modules
+						bDisplayContentFolderSuffix = false;
+					}
+
+					if (bDisplayContentFolderSuffix)
+					{
+						FolderDisplayNameOverride = FText::Format(LOCTEXT("ContentFolderDisplayNameFmt", "{0} Content"), FText::AsCultureInvariant(OverrideName));
+					}
+					else
+					{
+						FolderDisplayNameOverride = FText::AsCultureInvariant(OverrideName);
+					}
 				}
 			}
 		}
diff --git a/Engine/Source/Editor/ContentBrowserData/Private/ContentBrowserItemData.cpp b/Engine/Source/Editor/ContentBrowserData/Private/ContentBrowserItemData.cpp
--- a/Engine/Source/Editor/ContentBrowserData/Private/ContentBrowserItemData.cpp
+++ b/Engine/Source/Editor/ContentBrowserData/Private/ContentBrowserItemData.cpp
@@ -117,7 +117,17 @@ FName FContentBrowserItemData::GetInvariantPath() const
 	{
 		return Internal;
 	}
-	return VirtualPath; // None or a useful path
+	else if (!VirtualPath.IsNone())
+	{
+		if (UContentBrowserDataSource* DataSource = OwnerDataSource.Get())
+		{
+			FName ConvertedPath;
+			DataSource->TryConvertVirtualPath(VirtualPath, ConvertedPath);
+			return ConvertedPath;
+		}
+	}
+
+	return VirtualPath;
 }
 
 FName FContentBrowserItemData::GetInternalPath() const
diff --git a/Engine/Source/Editor/ContentBrowserData/Public/ContentBrowserDataSubsystem.h b/Engine/Source/Editor/ContentBrowserData/Public/ContentBrowserDataSubsystem.h
--- a/Engine/Source/Editor/ContentBrowserData/Public/ContentBrowserDataSubsystem.h
+++ b/Engine/Source/Editor/ContentBrowserData/Public/ContentBrowserDataSubsystem.h
@@ -31,6 +31,7 @@
 
 class FSubsystemCollectionBase;
 class FText;
+class IPlugin;
 class UContentBrowserDataSource;
 class UObject;
 struct FAssetData;
@@ -77,6 +78,9 @@ DECLARE_MULTICAST_DELEGATE(FOnContentBrowserItemDataDiscoveryComplete);
 /** Called when generating a virtual path, allows customization of how a virtual path is generated. */
 DECLARE_DELEGATE_TwoParams(FContentBrowserGenerateVirtualPathDelegate, const FStringView, FStringBuilderBase&);
 
+/** Called to check whether the plugin VersePath should be its path in the content browser, if it has one */
+DECLARE_DELEGATE_RetVal_OneParam(bool, FContentBrowserUsePluginVersePathDelegate, const TSharedRef<IPlugin>&);
+
 /** Internal - Filter data used to inject dummy items for the path down to the mount root of each data source */
 USTRUCT()
 struct CONTENTBROWSERDATA_API FContentBrowserCompiledSubsystemFilter
@@ -367,6 +371,21 @@ public:
 	 */
 	FContentBrowserGenerateVirtualPathDelegate& OnGenerateVirtualPathPrefix();
 
+	/**
+	 * Return whether the plugin VersePath should be its path in the content browser, if it has one
+	 */
+	bool UsePluginVersePath(const TSharedRef<IPlugin>& Plugin);
+
+	/**
+	 * Set delegate to check whether the plugin VersePath should be its path in the content browser, if it has one
+	 */
+	void SetUsePluginVersePathDelegate(FContentBrowserUsePluginVersePathDelegate InDelegate);
+
+	/**
+	 * Delegate to check whether the plugin VersePath should be its path in the content browser, if it has one
+	 */
+	FContentBrowserUsePluginVersePathDelegate& GetUsePluginVersePathDelegate();
+
 	/**
 	 * Prefix to use when generating virtual paths and "Show All Folder" option is enabled
 	 */
@@ -529,6 +548,11 @@ private:
 	 */
 	FContentBrowserGenerateVirtualPathDelegate GenerateVirtualPathPrefixDelegate;
 
+	/**
+	 * Delegate to check whether the plugin VersePath should be its path in the content browser, if it has one
+	 */
+	FContentBrowserUsePluginVersePathDelegate UsePluginVersePathDelegate;
+
 	/**
 	 * Optional array of invariant paths to use when sorting
 	 */
