Updated how we hide the /Game/Collections folder in the Content Browser to also support hiding private collections folders, and other folders if desired.

diff --git a/Engine/Plugins/Editor/ContentBrowser/ContentBrowserAssetDataSource/Source/ContentBrowserAssetDataSource/Private/ContentBrowserAssetDataSource.cpp b/Engine/Plugins/Editor/ContentBrowser/ContentBrowserAssetDataSource/Source/ContentBrowserAssetDataSource/Private/ContentBrowserAssetDataSource.cpp
--- a/Engine/Plugins/Editor/ContentBrowser/ContentBrowserAssetDataSource/Source/ContentBrowserAssetDataSource/Private/ContentBrowserAssetDataSource.cpp
+++ b/Engine/Plugins/Editor/ContentBrowser/ContentBrowserAssetDataSource/Source/ContentBrowserAssetDataSource/Private/ContentBrowserAssetDataSource.cpp
@@ -2313,11 +2313,11 @@ bool UContentBrowserAssetDataSource::PrioritizeSearchPath(const FName InPath)
 	return true;
 }
 
-bool UContentBrowserAssetDataSource::IsFolderVisible(const FName InPath, const EContentBrowserIsFolderVisibleFlags InFlags, TOptional<FContentBrowserFolderContentsFilter> InContentsFilter)
+bool UContentBrowserAssetDataSource::IsFolderVisible(const FName InPath, const EContentBrowserIsFolderVisibleFlags InFlags, const FContentBrowserFolderContentsFilter& InContentsFilter)
 {
 	const TSharedPtr<FPathPermissionList>& ShowPrivateContentPermissionList = IContentBrowserSingleton::Get().GetShowPrivateContentPermissionList();
 
-	auto IsInternalFolderVisible = [this, InContentsFilter, &ShowPrivateContentPermissionList](FName InternalFolderPath) {
+	auto IsInternalFolderVisible = [this, &InContentsFilter, &ShowPrivateContentPermissionList](FName InternalFolderPath) {
 		const EContentBrowserFolderAttributes FolderAttributes = GetAssetFolderAttributes(InternalFolderPath);
 		if (EnumHasAnyFlags(FolderAttributes, EContentBrowserFolderAttributes::AlwaysVisible))
 		{
@@ -2337,22 +2337,22 @@ bool UContentBrowserAssetDataSource::IsFolderVisible(const FName InPath, const E
 			}
 		}
 
-		if (InContentsFilter.IsSet())
+		if (EnumHasAnyFlags(InContentsFilter.ItemCategoryFilter, EContentBrowserItemCategoryFilter::IncludeAssets)
+			&& EnumHasAnyFlags(FolderAttributes, EContentBrowserFolderAttributes::HasAssets))
 		{
-			if (EnumHasAnyFlags(InContentsFilter->ItemCategoryFilter, EContentBrowserItemCategoryFilter::IncludeAssets)
-				&& EnumHasAnyFlags(FolderAttributes, EContentBrowserFolderAttributes::HasAssets))
-			{
-				return true;
-			}
-			if (EnumHasAnyFlags(InContentsFilter->ItemCategoryFilter, EContentBrowserItemCategoryFilter::IncludeRedirectors)
-				&& EnumHasAnyFlags(FolderAttributes, EContentBrowserFolderAttributes::HasRedirectors))
-			{
-				return true;
-			}
-			return false;
+			return true;
 		}
-
-		return true;
+		if (EnumHasAnyFlags(InContentsFilter.ItemCategoryFilter, EContentBrowserItemCategoryFilter::IncludeRedirectors)
+			&& EnumHasAnyFlags(FolderAttributes, EContentBrowserFolderAttributes::HasRedirectors))
+		{
+			return true;
+		}
+		if (InContentsFilter.HideFolderIfEmptyFilter)
+		{
+			FNameBuilder InternalFolderPathBuilder(InternalFolderPath);
+			return !InContentsFilter.HideFolderIfEmptyFilter->HideFolderIfEmpty(InternalFolderPath, InternalFolderPathBuilder);
+		}
+		return false;
 	};
 
 	FName ConvertedPath;
diff --git a/Engine/Plugins/Editor/ContentBrowser/ContentBrowserAssetDataSource/Source/ContentBrowserAssetDataSource/Public/ContentBrowserAssetDataSource.h b/Engine/Plugins/Editor/ContentBrowser/ContentBrowserAssetDataSource/Source/ContentBrowserAssetDataSource/Public/ContentBrowserAssetDataSource.h
--- a/Engine/Plugins/Editor/ContentBrowser/ContentBrowserAssetDataSource/Source/ContentBrowserAssetDataSource/Public/ContentBrowserAssetDataSource.h
+++ b/Engine/Plugins/Editor/ContentBrowser/ContentBrowserAssetDataSource/Source/ContentBrowserAssetDataSource/Public/ContentBrowserAssetDataSource.h
@@ -251,7 +251,7 @@ public:
 
 	virtual bool PrioritizeSearchPath(const FName InPath) override;
 
-	virtual bool IsFolderVisible(const FName InPath, const EContentBrowserIsFolderVisibleFlags InFlags, TOptional<FContentBrowserFolderContentsFilter> InContentsFilter) override;
+	virtual bool IsFolderVisible(const FName InPath, const EContentBrowserIsFolderVisibleFlags InFlags, const FContentBrowserFolderContentsFilter& InContentsFilter) override;
 
 	virtual bool CanCreateFolder(const FName InPath, FText* OutErrorMsg) override;
 
diff --git a/Engine/Plugins/Editor/ContentBrowser/ContentBrowserClassDataSource/Source/ContentBrowserClassDataSource/Private/ContentBrowserClassDataSource.cpp b/Engine/Plugins/Editor/ContentBrowser/ContentBrowserClassDataSource/Source/ContentBrowserClassDataSource/Private/ContentBrowserClassDataSource.cpp
--- a/Engine/Plugins/Editor/ContentBrowser/ContentBrowserClassDataSource/Source/ContentBrowserClassDataSource/Private/ContentBrowserClassDataSource.cpp
+++ b/Engine/Plugins/Editor/ContentBrowser/ContentBrowserClassDataSource/Source/ContentBrowserClassDataSource/Private/ContentBrowserClassDataSource.cpp
@@ -411,10 +411,10 @@ bool UContentBrowserClassDataSource::EnumerateItemsForObjects(const TArrayView<U
 	return true;
 }
 
-bool UContentBrowserClassDataSource::IsFolderVisible(const FName InPath, const EContentBrowserIsFolderVisibleFlags InFlags, TOptional<FContentBrowserFolderContentsFilter> InContentsFilter)
+bool UContentBrowserClassDataSource::IsFolderVisible(const FName InPath, const EContentBrowserIsFolderVisibleFlags InFlags, const FContentBrowserFolderContentsFilter& InContentsFilter)
 {
 	// We only contain classes, bail if the caller wants to filter out folders that contain only classes
-	if (InContentsFilter.IsSet() && !EnumHasAnyFlags(InContentsFilter->ItemCategoryFilter, EContentBrowserItemCategoryFilter::IncludeClasses))
+	if (!EnumHasAnyFlags(InContentsFilter.ItemCategoryFilter, EContentBrowserItemCategoryFilter::IncludeClasses))
 	{
 		return false;
 	}
@@ -445,7 +445,7 @@ bool UContentBrowserClassDataSource::IsFolderVisible(const FName InPath, const E
 	}
 
 	// Class flag was checked above - if we are filtering out folders that don't contain "class" elements, we are filtering out all empty folders from this provider
-	if (InContentsFilter.IsSet() && !NativeClassHierarchy->HasClasses(ConvertedPath, /*bRecursive*/ true))
+	if (!NativeClassHierarchy->HasClasses(ConvertedPath, /*bRecursive*/ true))
 	{
 		return false;
 	}
diff --git a/Engine/Plugins/Editor/ContentBrowser/ContentBrowserClassDataSource/Source/ContentBrowserClassDataSource/Public/ContentBrowserClassDataSource.h b/Engine/Plugins/Editor/ContentBrowser/ContentBrowserClassDataSource/Source/ContentBrowserClassDataSource/Public/ContentBrowserClassDataSource.h
--- a/Engine/Plugins/Editor/ContentBrowser/ContentBrowserClassDataSource/Source/ContentBrowserClassDataSource/Public/ContentBrowserClassDataSource.h
+++ b/Engine/Plugins/Editor/ContentBrowser/ContentBrowserClassDataSource/Source/ContentBrowserClassDataSource/Public/ContentBrowserClassDataSource.h
@@ -46,7 +46,7 @@ public:
 
 	virtual bool EnumerateItemsForObjects(const TArrayView<UObject*> InObjects, TFunctionRef<bool(FContentBrowserItemData&&)> InCallback) override;
 
-	virtual bool IsFolderVisible(const FName InPath, const EContentBrowserIsFolderVisibleFlags InFlags, TOptional<FContentBrowserFolderContentsFilter> InContentsFilter) override;
+	virtual bool IsFolderVisible(const FName InPath, const EContentBrowserIsFolderVisibleFlags InFlags, const FContentBrowserFolderContentsFilter& InContentsFilter) override;
 
 	virtual bool DoesItemPassFilter(const FContentBrowserItemData& InItem, const FContentBrowserDataCompiledFilter& InFilter) override;
 
diff --git a/Engine/Plugins/Editor/ContentBrowser/ContentBrowserFileDataSource/Source/ContentBrowserFileDataSource/Private/ContentBrowserFileDataSource.cpp b/Engine/Plugins/Editor/ContentBrowser/ContentBrowserFileDataSource/Source/ContentBrowserFileDataSource/Private/ContentBrowserFileDataSource.cpp
--- a/Engine/Plugins/Editor/ContentBrowser/ContentBrowserFileDataSource/Source/ContentBrowserFileDataSource/Private/ContentBrowserFileDataSource.cpp
+++ b/Engine/Plugins/Editor/ContentBrowser/ContentBrowserFileDataSource/Source/ContentBrowserFileDataSource/Private/ContentBrowserFileDataSource.cpp
@@ -1065,9 +1065,9 @@ bool UContentBrowserFileDataSource::PrioritizeSearchPath(const FName InPath)
 	return true;
 }
 
-bool UContentBrowserFileDataSource::IsFolderVisible(const FName InPath, const EContentBrowserIsFolderVisibleFlags InFlags, TOptional<FContentBrowserFolderContentsFilter> ContentsFilter)
+bool UContentBrowserFileDataSource::IsFolderVisible(const FName InPath, const EContentBrowserIsFolderVisibleFlags InFlags, const FContentBrowserFolderContentsFilter& ContentsFilter)
 {
-	if (ContentsFilter.IsSet() && !EnumHasAllFlags(ContentsFilter->ItemCategoryFilter, EContentBrowserItemCategoryFilter::IncludeMisc))
+	if (!EnumHasAllFlags(ContentsFilter.ItemCategoryFilter, EContentBrowserItemCategoryFilter::IncludeMisc))
 	{
 		return false;
 	}
@@ -1098,9 +1098,14 @@ bool UContentBrowserFileDataSource::IsFolderVisible(const FName InPath, const EC
 	{
 		return true;
 	}
-	// Misc flag was checked above - if we are filtering out folders that don't contain "misc" elements, we are filtering out all empty folders from this provider
-	if (ContentsFilter.IsSet() && EmptyFolders.ContainsByHash(InternalPathHash, InternalPathStrView))
+	// Misc flag was checked above - if we are filtering out folders that don't contain "misc" elements, we are filtering out all empty folders from this provider,
+	// unless a custom empty folder filter is provided, in which case we use that
+	if (EmptyFolders.ContainsByHash(InternalPathHash, InternalPathStrView))
 	{
+		if (ContentsFilter.HideFolderIfEmptyFilter)
+		{
+			return !ContentsFilter.HideFolderIfEmptyFilter->HideFolderIfEmpty(ConvertedPath, InternalPathStrView);
+		}
 		return false;
 	}
 	return true;
diff --git a/Engine/Plugins/Editor/ContentBrowser/ContentBrowserFileDataSource/Source/ContentBrowserFileDataSource/Public/ContentBrowserFileDataSource.h b/Engine/Plugins/Editor/ContentBrowser/ContentBrowserFileDataSource/Source/ContentBrowserFileDataSource/Public/ContentBrowserFileDataSource.h
--- a/Engine/Plugins/Editor/ContentBrowser/ContentBrowserFileDataSource/Source/ContentBrowserFileDataSource/Public/ContentBrowserFileDataSource.h
+++ b/Engine/Plugins/Editor/ContentBrowser/ContentBrowserFileDataSource/Source/ContentBrowserFileDataSource/Public/ContentBrowserFileDataSource.h
@@ -59,7 +59,7 @@ public:
 
 	virtual bool PrioritizeSearchPath(const FName InPath) override;
 
-	virtual bool IsFolderVisible(const FName InPath, const EContentBrowserIsFolderVisibleFlags InFlags, TOptional<FContentBrowserFolderContentsFilter> InContentsFilter) override;
+	virtual bool IsFolderVisible(const FName InPath, const EContentBrowserIsFolderVisibleFlags InFlags, const FContentBrowserFolderContentsFilter& InContentsFilter) override;
 
 	virtual bool DoesItemPassFilter(const FContentBrowserItemData& InItem, const FContentBrowserDataCompiledFilter& InFilter) override;
 
diff --git a/Engine/Source/Editor/ContentBrowser/Private/ContentBrowserUtils.cpp b/Engine/Source/Editor/ContentBrowser/Private/ContentBrowserUtils.cpp
--- a/Engine/Source/Editor/ContentBrowser/Private/ContentBrowserUtils.cpp
+++ b/Engine/Source/Editor/ContentBrowser/Private/ContentBrowserUtils.cpp
@@ -1412,6 +1412,8 @@ TArray<FContentBrowserItem> ContentBrowserUtils::GetChildItemsFromVirtualPath(
 	const SPathView& PathViewForFiltering 
 )
 {
+	UContentBrowserDataSubsystem* ContentBrowserData = IContentBrowserDataModule::Get().GetSubsystem();
+
 	const UContentBrowserSettings* ContentBrowserSettings = GetDefault<UContentBrowserSettings>();
 	bool bDisplayEmpty = ContentBrowserSettings->DisplayEmptyFolders;
 	// check to see if we have an instance config that overrides the default in UContentBrowserSettings
@@ -1422,14 +1424,15 @@ TArray<FContentBrowserItem> ContentBrowserUtils::GetChildItemsFromVirtualPath(
 	
 	EContentBrowserIsFolderVisibleFlags FolderFlags = ContentBrowserUtils::GetIsFolderVisibleFlags(bDisplayEmpty);
 
-	TOptional<FContentBrowserFolderContentsFilter> FolderFilter;
-	if (!bDisplayEmpty)
+	FContentBrowserFolderContentsFilter FolderFilter;
+	if (bDisplayEmpty)
 	{
-		FolderFilter = FContentBrowserFolderContentsFilter{};
-		FolderFilter->ItemCategoryFilter = ItemCategoryFilter;
+		FolderFilter.HideFolderIfEmptyFilter = ContentBrowserData->CreateHideFolderIfEmptyFilter();
+	}
+	else
+	{
+		FolderFilter.ItemCategoryFilter = ItemCategoryFilter;
 	}
-
-	UContentBrowserDataSubsystem* ContentBrowserData = IContentBrowserDataModule::Get().GetSubsystem();
 
 	FContentBrowserDataFilter SubItemsFilter;
 	SubItemsFilter.ItemTypeFilter = EContentBrowserItemTypeFilter::IncludeFolders;
diff --git a/Engine/Source/Editor/ContentBrowser/Private/SAssetView.cpp b/Engine/Source/Editor/ContentBrowser/Private/SAssetView.cpp
--- a/Engine/Source/Editor/ContentBrowser/Private/SAssetView.cpp
+++ b/Engine/Source/Editor/ContentBrowser/Private/SAssetView.cpp
@@ -270,10 +270,13 @@ public:
 		, FolderFilter()
 		, bDisplayEmptyFolders(AssetView->IsShowingEmptyFolders())
 	{
-		if (!bDisplayEmptyFolders)
+		if (bDisplayEmptyFolders)
 		{
-			FolderFilter = FContentBrowserFolderContentsFilter{};
-			FolderFilter->ItemCategoryFilter = InAssetView->DetermineItemCategoryFilter();
+			FolderFilter.HideFolderIfEmptyFilter = ContentBrowserData->CreateHideFolderIfEmptyFilter();
+		}
+		else
+		{
+			FolderFilter.ItemCategoryFilter = InAssetView->DetermineItemCategoryFilter();
 		}
 	}
 
@@ -333,7 +336,7 @@ public:
 private:
 	SAssetView* AssetView = nullptr;
 	UContentBrowserDataSubsystem* ContentBrowserData = nullptr;
-	TOptional<FContentBrowserFolderContentsFilter> FolderFilter;
+	FContentBrowserFolderContentsFilter FolderFilter;
 	const bool bDisplayEmptyFolders = true;
 };
 
diff --git a/Engine/Source/Editor/ContentBrowser/Private/SPathView.cpp b/Engine/Source/Editor/ContentBrowser/Private/SPathView.cpp
--- a/Engine/Source/Editor/ContentBrowser/Private/SPathView.cpp
+++ b/Engine/Source/Editor/ContentBrowser/Private/SPathView.cpp
@@ -364,7 +364,7 @@ protected:
 
 	struct FEmptyFolderFilter
 	{
-		TOptional<FContentBrowserFolderContentsFilter> FolderFilter;
+		FContentBrowserFolderContentsFilter FolderFilter;
 		EContentBrowserIsFolderVisibleFlags FolderFlags;
 	};
 	FEmptyFolderFilter GetEmptyFolderFilter(const FContentBrowserDataCompiledFilter& CompiledDataFilter) const;
@@ -401,15 +401,18 @@ FPathViewData::FEmptyFolderFilter FPathViewData::GetEmptyFolderFilter(const FCon
 		bDisplayEmpty = EditorConfig->bShowEmptyFolders;
 	}
 
-	UContentBrowserDataSubsystem* ContentBrowserData = IContentBrowserDataModule::Get().GetSubsystem();
-	TOptional<FContentBrowserFolderContentsFilter> FolderFilter;
-	if (!bDisplayEmpty)
+	FContentBrowserFolderContentsFilter FolderFilter;
+	if (bDisplayEmpty)
+	{
+		UContentBrowserDataSubsystem* ContentBrowserData = IContentBrowserDataModule::Get().GetSubsystem();
+		FolderFilter.HideFolderIfEmptyFilter = ContentBrowserData->CreateHideFolderIfEmptyFilter();
+	}
+	else
 	{
-		FolderFilter = FContentBrowserFolderContentsFilter{};
-		FolderFilter->ItemCategoryFilter = CompiledDataFilter.ItemCategoryFilter;
+		FolderFilter.ItemCategoryFilter = CompiledDataFilter.ItemCategoryFilter;
 	}
 	EContentBrowserIsFolderVisibleFlags FolderFlags = ContentBrowserUtils::GetIsFolderVisibleFlags(bDisplayEmpty);
-	return { FolderFilter, FolderFlags };
+	return { MoveTemp(FolderFilter), FolderFlags };
 }
 
 void FPathViewData::PopulateFullFolderTree(const FContentBrowserDataCompiledFilter& CompiledDataFilter)
diff --git a/Engine/Source/Editor/ContentBrowserData/Private/ContentBrowserDataSource.cpp b/Engine/Source/Editor/ContentBrowserData/Private/ContentBrowserDataSource.cpp
--- a/Engine/Source/Editor/ContentBrowserData/Private/ContentBrowserDataSource.cpp
+++ b/Engine/Source/Editor/ContentBrowserData/Private/ContentBrowserDataSource.cpp
@@ -263,7 +263,7 @@ bool UContentBrowserDataSource::PrioritizeSearchPath(const FName InPath)
 	return false;
 }
 
-bool UContentBrowserDataSource::IsFolderVisible(const FName Path, const EContentBrowserIsFolderVisibleFlags Flags, TOptional<FContentBrowserFolderContentsFilter> ContentsFilter)
+bool UContentBrowserDataSource::IsFolderVisible(const FName Path, const EContentBrowserIsFolderVisibleFlags Flags, const FContentBrowserFolderContentsFilter& ContentsFilter)
 {
 	return true;
 }
diff --git a/Engine/Source/Editor/ContentBrowserData/Private/ContentBrowserDataSubsystem.cpp b/Engine/Source/Editor/ContentBrowserData/Private/ContentBrowserDataSubsystem.cpp
--- a/Engine/Source/Editor/ContentBrowserData/Private/ContentBrowserDataSubsystem.cpp
+++ b/Engine/Source/Editor/ContentBrowserData/Private/ContentBrowserDataSubsystem.cpp
@@ -71,6 +71,69 @@ namespace ContentBrowserDataSubsystem
 			}
 		}
 	));
+
+	class FDefaultHideFolderIfEmptyFilter : public IContentBrowserHideFolderIfEmptyFilter
+	{
+	public:
+		FDefaultHideFolderIfEmptyFilter()
+			: GameDevelopersPath(FPackageName::FilenameToLongPackageName(FPaths::GameDevelopersDir()))
+		{
+		}
+
+		bool HideFolderIfEmpty(FName Path, FStringView PathString) const override
+		{
+			static const FName PublicCollectionsPath("/Game/Collections");
+
+			// Hide public collection folder.
+			if (Path == PublicCollectionsPath)
+			{
+				return true;
+			}
+
+			// Hide private collection folders.
+			if (PathString.StartsWith(GameDevelopersPath))
+			{
+				for (int32 Index = GameDevelopersPath.Len(); Index < PathString.Len(); ++Index)
+				{
+					// Scan past the developer name in /Game/Developers/<developer>/Collections.
+					if (PathString[Index] == TEXT('/'))
+					{
+						return PathString.RightChop(Index + 1).Equals(TEXT("Collections"));
+					}
+				}
+			}
+
+			return false;
+		}
+
+	private:
+		const FString GameDevelopersPath;
+	};
+
+	class FMergedHideFolderIfEmptyFilter : public IContentBrowserHideFolderIfEmptyFilter
+	{
+	public:
+		FMergedHideFolderIfEmptyFilter(TArray<TSharedPtr<IContentBrowserHideFolderIfEmptyFilter>>&& InHideFolderIfEmptyFilters)
+			: HideFolderIfEmptyFilters(MoveTemp(InHideFolderIfEmptyFilters))
+		{
+		}
+
+		bool HideFolderIfEmpty(FName Path, FStringView PathString) const override
+		{
+			for (const TSharedPtr<IContentBrowserHideFolderIfEmptyFilter>& HideFolderIfEmptyFilter : HideFolderIfEmptyFilters)
+			{
+				if (HideFolderIfEmptyFilter->HideFolderIfEmpty(Path, PathString))
+				{
+					return true;
+				}
+			}
+
+			return false;
+		}
+
+	private:
+		const TArray<TSharedPtr<IContentBrowserHideFolderIfEmptyFilter>> HideFolderIfEmptyFilters;
+	};
 }
 
 UContentBrowserDataSubsystem::UContentBrowserDataSubsystem()
@@ -111,6 +174,8 @@ void UContentBrowserDataSubsystem::Initialize(FSubsystemCollectionBase& Collecti
 	ModularFeatures.OnModularFeatureRegistered().AddUObject(this, &UContentBrowserDataSubsystem::HandleDataSourceRegistered);
 	ModularFeatures.OnModularFeatureUnregistered().AddUObject(this, &UContentBrowserDataSubsystem::HandleDataSourceUnregistered);
 
+	DefaultHideFolderIfEmptyFilter = MakeShared<ContentBrowserDataSubsystem::FDefaultHideFolderIfEmptyFilter>();
+
 	FEditorDelegates::BeginPIE.AddUObject(this, &UContentBrowserDataSubsystem::OnBeginPIE);
 	FEditorDelegates::EndPIE.AddUObject(this, &UContentBrowserDataSubsystem::OnEndPIE);
 
@@ -144,6 +209,7 @@ void UContentBrowserDataSubsystem::Deinitialize()
 	ActiveDataSources.Reset();
 	AvailableDataSources.Reset();
 	ActiveDataSourcesDiscoveringContent.Reset();
+	DefaultHideFolderIfEmptyFilter.Reset();
 
 	if (TickHandle.IsValid())
 	{
@@ -564,7 +630,7 @@ bool UContentBrowserDataSubsystem::PrioritizeSearchPath(const FName InPath)
 	return bDidPrioritize;
 }
 
-bool UContentBrowserDataSubsystem::IsFolderVisible(const FName InPath, const EContentBrowserIsFolderVisibleFlags InFlags, TOptional<FContentBrowserFolderContentsFilter> InContentsFilter) const
+bool UContentBrowserDataSubsystem::IsFolderVisible(const FName InPath, const EContentBrowserIsFolderVisibleFlags InFlags, const FContentBrowserFolderContentsFilter& InContentsFilter) const
 {
 	bool bIsKnownPath = false;
 	for (const auto& ActiveDataSourcePair : ActiveDataSources)
@@ -1096,6 +1162,55 @@ FContentBrowserUsePluginVersePathDelegate& UContentBrowserDataSubsystem::GetUseP
 	return UsePluginVersePathDelegate;
 }
 
+TSharedPtr<IContentBrowserHideFolderIfEmptyFilter> UContentBrowserDataSubsystem::CreateHideFolderIfEmptyFilter() const
+{
+	if (CreateHideFolderIfEmptyFilterDelegates.IsEmpty())
+	{
+		return DefaultHideFolderIfEmptyFilter;
+	}
+
+	TArray<TSharedPtr<IContentBrowserHideFolderIfEmptyFilter>> HideFolderIfEmptyFilters;
+	HideFolderIfEmptyFilters.Reserve(1 + CreateHideFolderIfEmptyFilterDelegates.Num());
+
+	if (DefaultHideFolderIfEmptyFilter)
+	{
+		HideFolderIfEmptyFilters.Add(DefaultHideFolderIfEmptyFilter);
+	}
+
+	for (const FContentBrowserCreateHideFolderIfEmptyFilter& CreateHideFolderIfEmptyFilter : CreateHideFolderIfEmptyFilterDelegates)
+	{
+		TSharedPtr<IContentBrowserHideFolderIfEmptyFilter> HideFolderIfEmptyFilter = CreateHideFolderIfEmptyFilter.Execute();
+		if (HideFolderIfEmptyFilter.IsValid())
+		{
+			HideFolderIfEmptyFilters.Add(MoveTemp(HideFolderIfEmptyFilter));
+		}
+	}
+
+	if (HideFolderIfEmptyFilters.Num() == 1)
+	{
+		return HideFolderIfEmptyFilters[0];
+	}
+
+	return MakeShared<ContentBrowserDataSubsystem::FMergedHideFolderIfEmptyFilter>(MoveTemp(HideFolderIfEmptyFilters));
+}
+
+FDelegateHandle UContentBrowserDataSubsystem::RegisterCreateHideFolderIfEmptyFilter(FContentBrowserCreateHideFolderIfEmptyFilter Delegate)
+{
+	return CreateHideFolderIfEmptyFilterDelegates.Add_GetRef(MoveTemp(Delegate)).GetHandle();
+}
+
+void UContentBrowserDataSubsystem::UnregisterCreateHideFolderIfEmptyFilter(FDelegateHandle DelegateHandle)
+{
+	int32 Index = CreateHideFolderIfEmptyFilterDelegates.IndexOfByPredicate([DelegateHandle](const FContentBrowserCreateHideFolderIfEmptyFilter& Delegate)
+		{
+			return DelegateHandle == Delegate.GetHandle();
+		});
+	if (Index != INDEX_NONE)
+	{
+		CreateHideFolderIfEmptyFilterDelegates.RemoveAtSwap(Index);
+	}
+}
+
 const FString& UContentBrowserDataSubsystem::GetAllFolderPrefix() const
 {
 	return AllFolderPrefix;
diff --git a/Engine/Source/Editor/ContentBrowserData/Public/ContentBrowserDataFilter.h b/Engine/Source/Editor/ContentBrowserData/Public/ContentBrowserDataFilter.h
--- a/Engine/Source/Editor/ContentBrowserData/Public/ContentBrowserDataFilter.h
+++ b/Engine/Source/Editor/ContentBrowserData/Public/ContentBrowserDataFilter.h
@@ -70,6 +70,19 @@ enum class EContentBrowserItemAttributeFilter : uint8
 };
 ENUM_CLASS_FLAGS(EContentBrowserItemAttributeFilter);
 
+/*
+ * Structure used to filter folders by their path
+ * @note All implementations of this filter should copy all dependent data on creation such that all calls to HideFolderIfEmpty for a given path return the same value for a instance of the implementation.
+ */
+class CONTENTBROWSERDATA_API IContentBrowserHideFolderIfEmptyFilter
+{
+public:
+	IContentBrowserHideFolderIfEmptyFilter() = default;
+	virtual ~IContentBrowserHideFolderIfEmptyFilter() = default;
+
+	virtual bool HideFolderIfEmpty(FName Path, FStringView PathString) const = 0;
+};
+
 /**
  * A list of typed filter structs and their associated data.
  * This allows systems to add new filter types that the core Content Browser data module doesn't know about.
@@ -365,20 +378,20 @@ public:
 /*
  * Structure used to optionally filter folders by their broad contents
  */
-USTRUCT(BlueprintType)
 struct CONTENTBROWSERDATA_API FContentBrowserFolderContentsFilter
 {
-	GENERATED_BODY()
-
 public:
-	UPROPERTY(EditAnywhere, BlueprintReadWrite, Category = "ContentBrowser")
 	EContentBrowserItemTypeFilter ItemTypeFilter = EContentBrowserItemTypeFilter::IncludeAll;
 
-	UPROPERTY(EditAnywhere, BlueprintReadWrite, Category = "ContentBrowser")
 	EContentBrowserItemCategoryFilter ItemCategoryFilter = EContentBrowserItemCategoryFilter::IncludeAll;
 
-	UPROPERTY(EditAnywhere, BlueprintReadWrite, Category = "ContentBrowser")
 	EContentBrowserItemAttributeFilter ItemAttributeFilter = EContentBrowserItemAttributeFilter::IncludeAll;
+
+	/**
+	* Used to filter out empty folders. All empty folders will get filtered out if this TSharedPtr is null.
+	* Intended for content folders that always exist but do not generally contain assets, such as /Game/Collections
+	*/
+	TSharedPtr<IContentBrowserHideFolderIfEmptyFilter> HideFolderIfEmptyFilter;
 };
 
 /**
diff --git a/Engine/Source/Editor/ContentBrowserData/Public/ContentBrowserDataSource.h b/Engine/Source/Editor/ContentBrowserData/Public/ContentBrowserDataSource.h
--- a/Engine/Source/Editor/ContentBrowserData/Public/ContentBrowserDataSource.h
+++ b/Engine/Source/Editor/ContentBrowserData/Public/ContentBrowserDataSource.h
@@ -263,7 +263,7 @@ public:
 	 * @param Flags Currently unused
 	 * @param ContentsFilter Optional filter for limiting visibility to only folders with a certain type of contents based on desired view filtering. If this is unset, folders may be visible regardless of their contents or empty status.
 	 */
-	virtual bool IsFolderVisible(const FName Path, const EContentBrowserIsFolderVisibleFlags Flags, TOptional<FContentBrowserFolderContentsFilter> ContentsFilter);
+	virtual bool IsFolderVisible(const FName Path, const EContentBrowserIsFolderVisibleFlags Flags, const FContentBrowserFolderContentsFilter& ContentsFilter);
 
 	/*
 	 * Query whether a folder can be created at the given virtual path, optionally providing error information if it cannot.
diff --git a/Engine/Source/Editor/ContentBrowserData/Public/ContentBrowserDataSubsystem.h b/Engine/Source/Editor/ContentBrowserData/Public/ContentBrowserDataSubsystem.h
--- a/Engine/Source/Editor/ContentBrowserData/Public/ContentBrowserDataSubsystem.h
+++ b/Engine/Source/Editor/ContentBrowserData/Public/ContentBrowserDataSubsystem.h
@@ -81,6 +81,9 @@ DECLARE_DELEGATE_TwoParams(FContentBrowserGenerateVirtualPathDelegate, const FSt
 /** Called to check whether the plugin VersePath should be its path in the content browser, if it has one */
 DECLARE_DELEGATE_RetVal_OneParam(bool, FContentBrowserUsePluginVersePathDelegate, const TSharedRef<IPlugin>&);
 
+/** Called to create a IContentBrowserHideFolderIfEmptyFilter filter */
+DECLARE_DELEGATE_RetVal(TSharedPtr<IContentBrowserHideFolderIfEmptyFilter>, FContentBrowserCreateHideFolderIfEmptyFilter);
+
 /** Internal - Filter data used to inject dummy items for the path down to the mount root of each data source */
 USTRUCT()
 struct CONTENTBROWSERDATA_API FContentBrowserCompiledSubsystemFilter
@@ -274,7 +277,7 @@ public:
 	 * @param Flags Currently unused
 	 * @param ContentsFilter Optional filter for limiting visibility to only folders with a certain type of contents based on desired view filtering.
 	 */
-	bool IsFolderVisible(const FName Path, const EContentBrowserIsFolderVisibleFlags Flags = EContentBrowserIsFolderVisibleFlags::Default, TOptional<FContentBrowserFolderContentsFilter> ContentsFilter = {}) const;
+	bool IsFolderVisible(const FName Path, const EContentBrowserIsFolderVisibleFlags Flags = EContentBrowserIsFolderVisibleFlags::Default, const FContentBrowserFolderContentsFilter& ContentsFilter = {}) const;
 
 	/**
 	 * Query whether the given virtual folder should be visible if the UI is asking to hide empty content folders.
@@ -386,6 +389,19 @@ public:
 	 */
 	FContentBrowserUsePluginVersePathDelegate& GetUsePluginVersePathDelegate();
 
+	/**
+	 * Creates the default IContentBrowserHideFolderIfEmptyFilter, and combines it with the results of any additional registered FContentBrowserCreateHideFolderIfEmptyFilter delegates.
+	 * @note The filter is only valid until any dependent data changes, so only keep it for a short time (typically within a function call tree).
+	 */
+	TSharedPtr<IContentBrowserHideFolderIfEmptyFilter> CreateHideFolderIfEmptyFilter() const;
+
+	/**
+	* Registers/deregisters a delegate used to create a filter used to filter out empty folders. If multiple delegates are registered they will get or'd togeather.
+	* Intended for content folders that always exist but do not generally contain assets, such as /Game/Collections
+	*/
+	FDelegateHandle RegisterCreateHideFolderIfEmptyFilter(FContentBrowserCreateHideFolderIfEmptyFilter Delegate);
+	void UnregisterCreateHideFolderIfEmptyFilter(FDelegateHandle DelegateHandle);
+
 	/**
 	 * Prefix to use when generating virtual paths and "Show All Folder" option is enabled
 	 */
@@ -553,6 +569,16 @@ private:
 	 */
 	FContentBrowserUsePluginVersePathDelegate UsePluginVersePathDelegate;
 
+	/**
+	 * Default filter used to check whether an empty folder should always be filtered out in the Content Browser
+	 */
+	TSharedPtr<IContentBrowserHideFolderIfEmptyFilter> DefaultHideFolderIfEmptyFilter;
+
+	/**
+	 * Delegates to create filters that check whether an empty folder should always be filtered out in the Content Browser
+	 */
+	TArray<FContentBrowserCreateHideFolderIfEmptyFilter> CreateHideFolderIfEmptyFilterDelegates;
+
 	/**
 	 * Optional array of invariant paths to use when sorting
 	 */
diff --git a/Engine/Source/Runtime/AssetRegistry/Private/AssetDataGatherer.cpp b/Engine/Source/Runtime/AssetRegistry/Private/AssetDataGatherer.cpp
--- a/Engine/Source/Runtime/AssetRegistry/Private/AssetDataGatherer.cpp
+++ b/Engine/Source/Runtime/AssetRegistry/Private/AssetDataGatherer.cpp
@@ -966,7 +966,7 @@ void FScanDir::SetScanResults(FStringView LocalAbsPath, const FInherited& Parent
 	{
 		FDiscoveredPathData& SubDirPath = InOutSubDirs[Index];
 		FScanDir& SubScanDir = FindOrAddSubDir(SubDirPath.RelPath);
-		bool bReportResult = SubScanDir.IsMonitored(Accumulated) && MountDir->GetDiscovery().ShouldDirBeReported(SubDirPath.LongPackageName);
+		bool bReportResult = SubScanDir.IsMonitored(Accumulated);
 		if (!bReportResult)
 		{
 			Swap(SubDirPath, InOutSubDirs.Last());
@@ -1573,8 +1573,6 @@ FAssetDataDiscovery::FAssetDataDiscovery(const TArray<FString>& InLongPackageNam
 	, IsPaused(0)
 	, NumDirectoriesToScan(0)
 {
-	DirLongPackageNamesToNotReport.Add(TEXT("/Game/Collections"));
-
 	if (bAsyncEnabled && !FPlatformProcess::SupportsMultithreading())
 	{
 		bAsyncEnabled = false;
@@ -2741,7 +2739,6 @@ SIZE_T FAssetDataDiscovery::GetAllocatedSize() const
 	SIZE_T Result = 0;
 	Result += GetArrayRecursiveAllocatedSize(LongPackageNamesDenyList);
 	Result += GetArrayRecursiveAllocatedSize(MountRelativePathsDenyList);
-	Result += GetArrayRecursiveAllocatedSize(DirLongPackageNamesToNotReport);
 	if (Thread)
 	{
 		// TODO: Thread->GetAllocatedSize()
@@ -2797,7 +2794,6 @@ void FAssetDataDiscovery::Shrink()
 	check(TickOwner.IsOwnedByCurrentThread());
 	CHECK_IS_LOCKED_CURRENT_THREAD(TreeLock);
 	CHECK_IS_LOCKED_CURRENT_THREAD(ResultsLock);
-	DirLongPackageNamesToNotReport.Shrink();
 	DiscoveredDirectories.Shrink();
 	DiscoveredFiles.Shrink();
 	DiscoveredSingleFiles.Shrink();
@@ -2934,12 +2930,6 @@ void FAssetDataDiscovery::OnDirectoryCreated(FStringView LocalAbsPath)
 		return;
 	}
 
-	// Skip reporting the directory if it is in the deny list of directories to not report
-	if (!ShouldDirBeReported(LongPackageName))
-	{
-		return;
-	}
-
 	FDiscoveredPathData DirData;
 	DirData.LocalAbsPath = LocalAbsPath;
 	DirData.LongPackageName = LongPackageName;
@@ -3144,10 +3134,6 @@ FAssetDataDiscovery::FDirectoryResult::FDirectoryResult(FStringView InDirAbsPath
 	}
 
 }
-bool FAssetDataDiscovery::ShouldDirBeReported(FStringView LongPackageName) const
-{
-	return !DirLongPackageNamesToNotReport.ContainsByHash(GetTypeHash(LongPackageName), LongPackageName);
-}
 
 // Reads an FMemoryView once
 class FMemoryViewReader
diff --git a/Engine/Source/Runtime/AssetRegistry/Private/AssetDataGatherer.h b/Engine/Source/Runtime/AssetRegistry/Private/AssetDataGatherer.h
--- a/Engine/Source/Runtime/AssetRegistry/Private/AssetDataGatherer.h
+++ b/Engine/Source/Runtime/AssetRegistry/Private/AssetDataGatherer.h
@@ -454,7 +454,7 @@ private:
 	/** Sum of all SearchTimes. */
 	float CumulativeGatherTime = 0.f;
 
-	/** The directories found during the search, unless they are hidden by DirLongPackageNamesToNotReport. */
+	/** The directories found during the search. */
 	TArray<FString> DiscoveredPaths;
 
 	/** The time spent in TickInternal since the last idle time. Used for performance metrics when reporting results. */
diff --git a/Engine/Source/Runtime/AssetRegistry/Private/AssetDataGathererPrivate.h b/Engine/Source/Runtime/AssetRegistry/Private/AssetDataGathererPrivate.h
--- a/Engine/Source/Runtime/AssetRegistry/Private/AssetDataGathererPrivate.h
+++ b/Engine/Source/Runtime/AssetRegistry/Private/AssetDataGathererPrivate.h
@@ -675,15 +675,6 @@ private:
 	 */
 	static EGatherableFileType GetFileType(FStringView FilePath);
 
-	/**
-	 * Return whether a directory with the given LongPackageName should be reported to the AssetRegistry
-	 * We do not report some directories because they should not enter the AssetRegistry list of paths if empty,
-	 * and reporting a path to the AssetRegistry adds it unconditionally to the list of paths.
-	 * If ShouldDirBeReported returns false, the directory will still be added to the catalog if non-empty,
-	 * because the AssetRegistry adds the path of every added file.
-	 */
-	bool ShouldDirBeReported(FStringView LongPackageName) const;
-
 	/** Handle the actions necessary for a single created file. */
 	void OnFileCreated(const FString& LocalPaths);
 	/** Helper for TrySetDirectoryProperties, called from within the critical section. */
@@ -730,8 +721,6 @@ private:
 	TArray<FString> LongPackageNamesDenyList;
 	/** DenyList of relative paths in each mount. Child paths will not be scanned unless requested to ignore DenyLists. */
 	TArray<FString> MountRelativePathsDenyList;
-	/** LongPackageNames for directories that should not be reported, see ShouldDirBeReported. */
-	TSet<FString> DirLongPackageNamesToNotReport;
 	/** Thread to run the discovery FRunnable on. Read-only while threading is possible. */
 	FRunnableThread* Thread;
 	/**
