Fix a crash in the Content Browser and Asset Dialog caused by incorrect usage of the FContentBrowserItemPath constructor.

diff --git a/Engine/Source/Editor/ContentBrowser/Private/Experimental/ContentBrowserExtensionUtils.cpp b/Engine/Source/Editor/ContentBrowser/Private/Experimental/ContentBrowserExtensionUtils.cpp
--- a/Engine/Source/Editor/ContentBrowser/Private/Experimental/ContentBrowserExtensionUtils.cpp
+++ b/Engine/Source/Editor/ContentBrowser/Private/Experimental/ContentBrowserExtensionUtils.cpp
@@ -23,7 +23,12 @@ namespace UE::Editor::ContentBrowser::ExtensionUtils
 			}
 		}
 
-		return ContentBrowserUtils::GetPathColor(FolderPath.GetInternalPathString());
+		if (FolderPath.HasInternalPath())
+		{
+			return ContentBrowserUtils::GetPathColor(FolderPath.GetInternalPathString());
+		}
+
+		return TOptional<FLinearColor>();
 	}
 	
 	TOptional<FLinearColor> GetFolderColor(const FName& FolderPath)
diff --git a/Engine/Source/Editor/ContentBrowser/Private/SAssetDialog.cpp b/Engine/Source/Editor/ContentBrowser/Private/SAssetDialog.cpp
--- a/Engine/Source/Editor/ContentBrowser/Private/SAssetDialog.cpp
+++ b/Engine/Source/Editor/ContentBrowser/Private/SAssetDialog.cpp
@@ -1037,7 +1037,7 @@ TSharedRef<SWidget> SAssetDialog::OnGetCrumbDelimiterContent(const FString& Crum
 
 			if (FolderBrushName != NAME_None)
 			{
-				const FContentBrowserItemPath FolderPath(SubItem.GetVirtualPath(), SubItem.GetInternalPath());
+				const FContentBrowserItemPath FolderPath = SubItem.GetItemPath();
 				FLinearColor FolderColor = UE::Editor::ContentBrowser::ExtensionUtils::GetFolderColor(FolderPath).Get(ContentBrowserUtils::GetDefaultColor());
 				
 				FMenuEntryParams Params;
diff --git a/Engine/Source/Editor/ContentBrowser/Private/SContentBrowser.cpp b/Engine/Source/Editor/ContentBrowser/Private/SContentBrowser.cpp
--- a/Engine/Source/Editor/ContentBrowser/Private/SContentBrowser.cpp
+++ b/Engine/Source/Editor/ContentBrowser/Private/SContentBrowser.cpp
@@ -3042,7 +3042,7 @@ TSharedRef<SWidget> SContentBrowser::OnGetCrumbDelimiterContent(const FString& C
 
 				if (FolderBrushName != NAME_None)
 				{
-					const FContentBrowserItemPath FolderPath(SubItem.GetVirtualPath(), SubItem.GetInternalPath());
+					const FContentBrowserItemPath FolderPath = SubItem.GetItemPath();
 					FLinearColor FolderColor = UE::Editor::ContentBrowser::ExtensionUtils::GetFolderColor(FolderPath).Get(ContentBrowserUtils::GetDefaultColor());
 						
 					FMenuEntryParams Params;
diff --git a/Engine/Source/Editor/ContentBrowserData/Private/ContentBrowserItem.cpp b/Engine/Source/Editor/ContentBrowserData/Private/ContentBrowserItem.cpp
--- a/Engine/Source/Editor/ContentBrowserData/Private/ContentBrowserItem.cpp
+++ b/Engine/Source/Editor/ContentBrowserData/Private/ContentBrowserItem.cpp
@@ -4,6 +4,7 @@
 #include "ContentBrowserDataSource.h"
 #include "IContentBrowserDataModule.h"
 #include "ContentBrowserDataSubsystem.h"
+#include "ContentBrowserItemPath.h"
 
 #define LOCTEXT_NAMESPACE "ContentBrowserData"
 
@@ -629,6 +630,29 @@ bool FContentBrowserItem::Legacy_TryGetAssetData(FAssetData& OutAssetData) const
 	return FContentBrowserItemHelper::CallDataSourceImpl<UContentBrowserDataSource>(*this, UE_PROJECTION_MEMBER(UContentBrowserDataSource, Legacy_TryGetAssetData), OutAssetData);
 }
 
+FContentBrowserItemPath FContentBrowserItem::GetItemPath() const
+{
+	const FName VirtualPath = GetVirtualPath();
+	const FName InternalPath = GetInternalPath();
+
+	if (VirtualPath.IsNone() && InternalPath.IsNone())
+	{
+		return FContentBrowserItemPath();
+	}
+
+	if (InternalPath.IsNone())
+	{
+		return FContentBrowserItemPath(VirtualPath, EContentBrowserPathType::Virtual);
+	}
+	
+	if (VirtualPath.IsNone())
+	{
+		return FContentBrowserItemPath(InternalPath, EContentBrowserPathType::Internal);
+	}
+	
+	return FContentBrowserItemPath(VirtualPath, InternalPath);
+}
+
 
 bool UContentBrowserItemLibrary::IsFolder(const FContentBrowserItem& Item)
 {
diff --git a/Engine/Source/Editor/ContentBrowserData/Public/ContentBrowserItem.h b/Engine/Source/Editor/ContentBrowserData/Public/ContentBrowserItem.h
--- a/Engine/Source/Editor/ContentBrowserData/Public/ContentBrowserItem.h
+++ b/Engine/Source/Editor/ContentBrowserData/Public/ContentBrowserItem.h
@@ -19,6 +19,7 @@
 
 #include "ContentBrowserItem.generated.h"
 
+struct FContentBrowserItemPath;
 class FAssetThumbnail;
 class FContentBrowserMinimalItemData;
 class IContentBrowserHideFolderIfEmptyFilter;
@@ -477,6 +478,12 @@ public:
 	 */
 	bool Legacy_TryGetAssetData(FAssetData& OutAssetData) const;
 
+	/**
+	 * Creates a ContentBrowserItemPath, potentially looking up internal paths.
+	 * @return The Item Path.
+	 */
+	FContentBrowserItemPath GetItemPath() const;
+
 private:
 	/**
 	 * Array of internal items that comprise this composite item.
