PR #12068: Ensuring the DataSources are Deactivated when the UContentBrowserDataSubsystem is Deinitialized

diff --git a/Engine/Source/Editor/ContentBrowserData/Private/ContentBrowserDataSubsystem.cpp b/Engine/Source/Editor/ContentBrowserData/Private/ContentBrowserDataSubsystem.cpp
--- a/Engine/Source/Editor/ContentBrowserData/Private/ContentBrowserDataSubsystem.cpp
+++ b/Engine/Source/Editor/ContentBrowserData/Private/ContentBrowserDataSubsystem.cpp
@@ -206,6 +206,7 @@ void UContentBrowserDataSubsystem::Deinitialize()
 
 	FPackageName::OnContentPathMounted().RemoveAll(this);
 
+	DeactivateAllDataSources();
 	ActiveDataSources.Reset();
 	AvailableDataSources.Reset();
 	ActiveDataSourcesDiscoveringContent.Reset();
