[ContentBrowser] Update DragDrop folder widget style
Missed file from a previous review for the Folder DragDrop widget style

diff --git a/Engine/Source/Editor/ContentBrowserData/Private/ContentBrowserDataDragDropOp.cpp b/Engine/Source/Editor/ContentBrowserData/Private/ContentBrowserDataDragDropOp.cpp
--- a/Engine/Source/Editor/ContentBrowserData/Private/ContentBrowserDataDragDropOp.cpp
+++ b/Engine/Source/Editor/ContentBrowserData/Private/ContentBrowserDataDragDropOp.cpp
@@ -12,6 +12,7 @@
 #include "HAL/PlatformCrt.h"
 #include "Templates/UnrealTemplate.h"
 #include "ThumbnailRendering/ThumbnailManager.h"
+#include "Widgets/Layout/SBox.h"
 #include "UObject/NameTypes.h"
 
 class UActorFactory;
@@ -27,8 +28,8 @@ TSharedRef<FContentBrowserDataDragDropOp> FContentBrowserDataDragDropOp::New(TAr
 
 	if (Operation->ShouldOverrideThumbnailWidget())
 	{
-		const TSharedPtr<SWidget> CustomThumbnailWidget = Operation->GetFolderWidget();
-		if (CustomThumbnailWidget.IsValid())
+		const TSharedPtr<SWidget> CustomThumbnailWidget = Operation->GetFolderWidgetDragAndDrop();
+		if (CustomThumbnailWidget && CustomThumbnailWidget.IsValid())
 		{
 			Operation->SetCustomThumbnailWidget(CustomThumbnailWidget.ToSharedRef());
 		}
@@ -183,11 +184,13 @@ bool FContentBrowserDataDragDropOp::ShouldOverrideThumbnailWidget() const
 
 TSharedPtr<SWidget> FContentBrowserDataDragDropOp::GetFolderWidget() const
 {
-	const FSlateBrush* FolderBrush = FAppStyle::GetBrush(FolderShadowBrushName);
-	if (FolderBrush != FAppStyle::GetDefaultBrush())
+	const FSlateBrush* FolderBrush = FAppStyle::GetOptionalBrush(FolderBrushName, nullptr, FAppStyle::GetBrush("ContentBrowser.ListViewFolderIcon"));
+	const FSlateBrush* FolderShadowBrush = FAppStyle::GetOptionalBrush(FolderShadowBrushName, nullptr , FAppStyle::GetBrush("ContentBrowser.FolderItem.DropShadow"));
+
+	if (FolderBrush && FolderShadowBrush)
 	{
 		return SNew(SBorder)
-			.BorderImage(FolderBrush)
+			.BorderImage(FolderShadowBrush)
 			.Padding(FMargin(0,0,2.0f,2.0f))
 			[
 				SNew(SImage)
@@ -197,3 +200,31 @@ TSharedPtr<SWidget> FContentBrowserDataDragDropOp::GetFolderWidget() const
 	}
 	return nullptr;
 }
+
+TSharedPtr<SWidget> FContentBrowserDataDragDropOp::GetFolderWidgetDragAndDrop() const
+{
+	const TSharedPtr<SWidget> FolderWidget = GetFolderWidget();
+
+	if (!FolderWidget || !FolderWidget.IsValid())
+	{
+		return nullptr;
+	}
+
+	static const FMargin FolderBoxPadding = FMargin(3.f, 4.f, 1.f, 0.f);
+	static constexpr float FolderSize = 60.f;
+
+	return SNew(SBorder)
+			.Padding(0.f)
+			.BorderImage(FAppStyle::GetBrush("ContentBrowser.ThumbnailDragDropBackground"))
+			.HAlign(HAlign_Center)
+			.VAlign(VAlign_Center)
+			[
+				SNew(SBox)
+				.WidthOverride(FolderSize)
+				.HeightOverride(FolderSize)
+				.Padding(FolderBoxPadding)
+				[
+					FolderWidget.ToSharedRef()
+				]
+			];
+}
