Updated headers using UnrealCodeFixup to make sure dllstorage is on methods/staticvars instead of types. Part 5/n

diff --git a/Engine/Source/Editor/ContentBrowserData/Public/ContentBrowserDataSubsystem.h b/Engine/Source/Editor/ContentBrowserData/Public/ContentBrowserDataSubsystem.h
--- a/Engine/Source/Editor/ContentBrowserData/Public/ContentBrowserDataSubsystem.h
+++ b/Engine/Source/Editor/ContentBrowserData/Public/ContentBrowserDataSubsystem.h
@@ -29,6 +29,8 @@
 
 #include "ContentBrowserDataSubsystem.generated.h"
 
+#define UE_API CONTENTBROWSERDATA_API
+
 class FSubsystemCollectionBase;
 class FText;
 class IPlugin;
@@ -86,7 +88,7 @@ DECLARE_DELEGATE_RetVal(TSharedPtr<IContentBrowserHideFolderIfEmptyFilter>, FCon
 
 /** Internal - Filter data used to inject dummy items for the path down to the mount root of each data source */
 USTRUCT()
-struct CONTENTBROWSERDATA_API FContentBrowserCompiledSubsystemFilter
+struct FContentBrowserCompiledSubsystemFilter
 {
 	GENERATED_BODY()
 	
@@ -96,7 +98,7 @@ public:
 
 /** Internal - Filter data used to inject dummy items */
 USTRUCT()
-struct CONTENTBROWSERDATA_API FContentBrowserCompiledVirtualFolderFilter
+struct FContentBrowserCompiledVirtualFolderFilter
 {
 	GENERATED_BODY()
 
@@ -108,101 +110,101 @@ public:
  * Subsystem that provides access to Content Browser data.
  * This type deals with the composition of multiple data sources, which provide information about the folders and files available in the Content Browser.
  */
-UCLASS(config=Editor)
-class CONTENTBROWSERDATA_API UContentBrowserDataSubsystem : public UEditorSubsystem, public IContentBrowserItemDataSink
+UCLASS(MinimalAPI, config=Editor)
+class UContentBrowserDataSubsystem : public UEditorSubsystem, public IContentBrowserItemDataSink
 {
 	GENERATED_BODY()
 
 public:
 	friend class FScopedSuppressContentBrowserDataTick;
 
-	UContentBrowserDataSubsystem();
+	UE_API UContentBrowserDataSubsystem();
 	//~ UEditorSubsystem interface
-	virtual void Initialize(FSubsystemCollectionBase& Collection) override;
-	virtual void Deinitialize() override;
-	virtual void ConvertInternalPathToVirtual(const FStringView InPath, FStringBuilderBase& OutPath) override;
+	UE_API virtual void Initialize(FSubsystemCollectionBase& Collection) override;
+	UE_API virtual void Deinitialize() override;
+	UE_API virtual void ConvertInternalPathToVirtual(const FStringView InPath, FStringBuilderBase& OutPath) override;
 
 	/**
 	 * Attempt to activate the named data source.
 	 * @return True if the data source was available and not already active, false otherwise.
 	 */
 	UFUNCTION(BlueprintCallable, Category="ContentBrowser")
-	bool ActivateDataSource(const FName Name);
+	UE_API bool ActivateDataSource(const FName Name);
 
 	/**
 	 * Attempt to deactivate the named data source.
 	 * @return True if the data source was available and active, false otherwise.
 	 */
 	UFUNCTION(BlueprintCallable, Category="ContentBrowser")
-	bool DeactivateDataSource(const FName Name);
+	UE_API bool DeactivateDataSource(const FName Name);
 
 	/**
 	 * Activate all available data sources.
 	 */
 	UFUNCTION(BlueprintCallable, Category="ContentBrowser")
-	void ActivateAllDataSources();
+	UE_API void ActivateAllDataSources();
 
 	/**
 	 * Deactivate all active data sources.
 	 */
 	UFUNCTION(BlueprintCallable, Category="ContentBrowser")
-	void DeactivateAllDataSources();
+	UE_API void DeactivateAllDataSources();
 
 	/**
 	 * Get the list of current available data sources.
 	 */
 	UFUNCTION(BlueprintCallable, Category="ContentBrowser")
-	TArray<FName> GetAvailableDataSources() const;
+	UE_API TArray<FName> GetAvailableDataSources() const;
 
 	/**
 	 * Get the list of current active data sources.
 	 */
 	UFUNCTION(BlueprintCallable, Category="ContentBrowser")
-	TArray<FName> GetActiveDataSources() const;
+	UE_API TArray<FName> GetActiveDataSources() const;
 
 	/**
 	 * Delegate called for incremental item data updates from data sources that can provide delta-updates.
 	 */
-	FOnContentBrowserItemDataUpdated& OnItemDataUpdated();
+	UE_API FOnContentBrowserItemDataUpdated& OnItemDataUpdated();
 
 	/**
 	 * Delegate called for wholesale item data updates from data sources that can't provide delta-updates, or when the set of active data sources is modified.
 	 */
-	FOnContentBrowserItemDataRefreshed& OnItemDataRefreshed();
+	UE_API FOnContentBrowserItemDataRefreshed& OnItemDataRefreshed();
 
 	/**
 	 * Delegate called when all active data sources have completed their initial content discovery scan.
 	 * @note May be called multiple times if new data sources are registered after the current set of active data sources have completed their initial scan.
 	 */
-	FOnContentBrowserItemDataDiscoveryComplete& OnItemDataDiscoveryComplete();
+	UE_API FOnContentBrowserItemDataDiscoveryComplete& OnItemDataDiscoveryComplete();
 
 	/**
 	 * Take a raw data filter and convert it into a compiled version that could be re-used for multiple queries using the same data (typically this is only useful for post-filtering multiple items).
 	 * @note The compiled filter is only valid until the data source changes, so only keep it for a short time (typically within a function call, or 1-frame).
 	 */
-	void CompileFilter(const FName InPath, const FContentBrowserDataFilter& InFilter, FContentBrowserDataCompiledFilter& OutCompiledFilter) const;
+	UE_API void CompileFilter(const FName InPath, const FContentBrowserDataFilter& InFilter, FContentBrowserDataCompiledFilter& OutCompiledFilter) const;
 
 	/**
 	 * Enumerate the items (folders and/or files) that match a previously compiled filter.
 	 */
-	void EnumerateItemsMatchingFilter(
+	UE_API void EnumerateItemsMatchingFilter(
 		const FContentBrowserDataCompiledFilter& InFilter, TFunctionRef<bool(FContentBrowserItem&&)> InCallback) const;
-	void EnumerateItemsMatchingFilter(const FContentBrowserDataCompiledFilter& InFilter,
+	UE_API void EnumerateItemsMatchingFilter(const FContentBrowserDataCompiledFilter& InFilter,
 		TFunctionRef<bool(FContentBrowserItemData&&)> InCallback) const;
-	void EnumerateItemsMatchingFilter(
+	UE_API void EnumerateItemsMatchingFilter(
 		const FContentBrowserDataCompiledFilter& InFilter, const TGetOrEnumerateSink<FContentBrowserItemData>& InSink) const;
 
 	/**
 	 * Enumerate the items (folders and/or files) that exist under the given virtual path.
 	 */
-	void EnumerateItemsUnderPath(const FName InPath,
+	UE_API void EnumerateItemsUnderPath(const FName InPath,
 		const FContentBrowserDataFilter& InFilter,
 		TFunctionRef<bool(FContentBrowserItem&&)> InCallback) const;
-	void EnumerateItemsUnderPath(const FName InPath,
+	UE_API void EnumerateItemsUnderPath(const FName InPath,
 		const FContentBrowserDataFilter& InFilter,
 		TFunctionRef<bool(FContentBrowserItemData&&)> InCallback) const;
 	// Version which allows passing an array or a callback, allows sources to presize the array for large queries
-	void EnumerateItemsUnderPath(const FName InPath,
+	UE_API void EnumerateItemsUnderPath(const FName InPath,
 		const FContentBrowserDataFilter& InFilter,
 		const TGetOrEnumerateSink<FContentBrowserItemData>& InSink) const;
 
@@ -210,14 +212,14 @@ public:
 	 * Get the items (folders and/or files) that exist under the given virtual path.
 	 */
 	UFUNCTION(BlueprintCallable, Category="ContentBrowser")
-	TArray<FContentBrowserItem> GetItemsUnderPath(const FName InPath, const FContentBrowserDataFilter& InFilter) const;
+	UE_API TArray<FContentBrowserItem> GetItemsUnderPath(const FName InPath, const FContentBrowserDataFilter& InFilter) const;
 
 	/**
 	 * Enumerate the items (folders and/or files) that exist at the given virtual path.
 	 * @note Multiple items may have the same virtual path if they are different types, or come from different data sources.
 	 */
-	void EnumerateItemsAtPath(const FName InPath, const EContentBrowserItemTypeFilter InItemTypeFilter, TFunctionRef<bool(FContentBrowserItem&&)> InCallback) const;
-	void EnumerateItemsAtPath(const FName InPath, const EContentBrowserItemTypeFilter InItemTypeFilter, TFunctionRef<bool(FContentBrowserItemData&&)> InCallback) const;
+	UE_API void EnumerateItemsAtPath(const FName InPath, const EContentBrowserItemTypeFilter InItemTypeFilter, TFunctionRef<bool(FContentBrowserItem&&)> InCallback) const;
+	UE_API void EnumerateItemsAtPath(const FName InPath, const EContentBrowserItemTypeFilter InItemTypeFilter, TFunctionRef<bool(FContentBrowserItemData&&)> InCallback) const;
 
 	/**
 	 * Enumerate the items (files) that exist at the given paths.
@@ -227,7 +229,7 @@ public:
 	 * @param InItemTypeFilter The types of items we want to find.
 	 * @param InCallback The function to invoke for each matching item (return true to continue enumeration).
 	 */
-	bool EnumerateItemsAtPaths(const TArrayView<struct FContentBrowserItemPath> InItemPaths, const EContentBrowserItemTypeFilter InItemTypeFilter, TFunctionRef<bool(FContentBrowserItemData&&)> InCallback) const;
+	UE_API bool EnumerateItemsAtPaths(const TArrayView<struct FContentBrowserItemPath> InItemPaths, const EContentBrowserItemTypeFilter InItemTypeFilter, TFunctionRef<bool(FContentBrowserItemData&&)> InCallback) const;
 
 	/**
 	 * Enumerate the items (files) that exist for the given objects.
@@ -236,20 +238,20 @@ public:
 	 * @param InObjects The objects to enumerate
 	 * @param InCallback The function to invoke for each matching item (return true to continue enumeration).
 	 */
-	bool EnumerateItemsForObjects(const TArrayView<UObject*> InObjects, TFunctionRef<bool(FContentBrowserItemData&&)> InCallback) const;
+	UE_API bool EnumerateItemsForObjects(const TArrayView<UObject*> InObjects, TFunctionRef<bool(FContentBrowserItemData&&)> InCallback) const;
 
 	/**
 	 * Get the items (folders and/or files) that exist at the given virtual path.
 	 * @note Multiple items may have the same virtual path if they are different types, or come from different data sources.
 	 */
 	UFUNCTION(BlueprintCallable, Category="ContentBrowser")
-	TArray<FContentBrowserItem> GetItemsAtPath(const FName InPath, const EContentBrowserItemTypeFilter InItemTypeFilter) const;
+	UE_API TArray<FContentBrowserItem> GetItemsAtPath(const FName InPath, const EContentBrowserItemTypeFilter InItemTypeFilter) const;
 
 	/**
 	 * Get the first item (folder and/or file) that exists at the given virtual path.
 	 */
 	UFUNCTION(BlueprintCallable, Category="ContentBrowser")
-	FContentBrowserItem GetItemAtPath(const FName InPath, const EContentBrowserItemTypeFilter InItemTypeFilter) const;
+	UE_API FContentBrowserItem GetItemAtPath(const FName InPath, const EContentBrowserItemTypeFilter InItemTypeFilter) const;
 
 	/**
 	 * Get a list of other paths that the data source may be using to represent a specific path
@@ -257,26 +259,26 @@ public:
 	 * @param The internal path (or object path) of an asset to get aliases for
 	 * @return All alternative paths that represent the input path (not including the input path itself)
 	 */
-	TArray<FContentBrowserItemPath> GetAliasesForPath(const FSoftObjectPath& InInternalPath) const;
-	TArray<FContentBrowserItemPath> GetAliasesForPath(const FContentBrowserItemPath InPath) const;
-	TArray<FContentBrowserItemPath> GetAliasesForPath(const FName InInternalPath) const;
+	UE_API TArray<FContentBrowserItemPath> GetAliasesForPath(const FSoftObjectPath& InInternalPath) const;
+	UE_API TArray<FContentBrowserItemPath> GetAliasesForPath(const FContentBrowserItemPath InPath) const;
+	UE_API TArray<FContentBrowserItemPath> GetAliasesForPath(const FName InInternalPath) const;
 
 	/**
 	 * Query whether any data sources are currently discovering content, and retrieve optional status messages that can be shown in the UI.
 	 */
-	bool IsDiscoveringItems(TArray<FText>* OutStatus = nullptr) const;
+	UE_API bool IsDiscoveringItems(TArray<FText>* OutStatus = nullptr) const;
 
 	/**
 	 * If possible, attempt to prioritize content discovery for the given virtual path.
 	 */
-	bool PrioritizeSearchPath(const FName InPath);
+	UE_API bool PrioritizeSearchPath(const FName InPath);
 
 	/**
 	 * Query whether the given virtual folder should be visible in the UI.
 	 * @param Path Virtual path of folder e.g. /All/GameData/Stuff
 	 * @param Flags Currently unused
 	 */
-	bool IsFolderVisible(const FName Path, const EContentBrowserIsFolderVisibleFlags Flags = EContentBrowserIsFolderVisibleFlags::Default) const;
+	UE_API bool IsFolderVisible(const FName Path, const EContentBrowserIsFolderVisibleFlags Flags = EContentBrowserIsFolderVisibleFlags::Default) const;
 
 	/**
 	 * Query whether the given virtual folder should be visible in the UI.
@@ -284,13 +286,13 @@ public:
 	 * @param Flags Currently unused
 	 * @param ContentsFilter Filter for limiting visibility to only folders with a certain type of contents based on desired view filtering.
 	 */
-	bool IsFolderVisible(const FName Path, const EContentBrowserIsFolderVisibleFlags Flags, const FContentBrowserFolderContentsFilter& ContentsFilter) const;
+	UE_API bool IsFolderVisible(const FName Path, const EContentBrowserIsFolderVisibleFlags Flags, const FContentBrowserFolderContentsFilter& ContentsFilter) const;
 
 	/**
 	 * Query whether the given virtual folder should be visible if the UI is asking to hide empty content folders.
 	 */
 	UE_DEPRECATED(5.3, "IsFolderVisibleIfHidingEmpty is deprecated. Use IsFolderVisible instead and add EContentBrowserIsFolderVisibleFlags::HideEmptyFolders to the flags.")
-	bool IsFolderVisibleIfHidingEmpty(const FName InPath) const;
+	UE_API bool IsFolderVisibleIfHidingEmpty(const FName InPath) const;
 
 	/*
 	 * Query whether a folder can be created at the given virtual path, optionally providing error information if it cannot.
@@ -300,7 +302,7 @@ public:
 	 *
 	 * @return True if the folder can be created, false otherwise.
 	 */
-	bool CanCreateFolder(const FName InPath, FText* OutErrorMsg) const;
+	UE_API bool CanCreateFolder(const FName InPath, FText* OutErrorMsg) const;
 
 	/*
 	 * Attempt to begin the process of asynchronously creating a folder at the given virtual path, returning a temporary item that can be finalized or canceled by the user.
@@ -309,7 +311,7 @@ public:
 	 *
 	 * @return The pending folder item to create (test for validity).
 	 */
-	FContentBrowserItemTemporaryContext CreateFolder(const FName InPath) const;
+	UE_API FContentBrowserItemTemporaryContext CreateFolder(const FName InPath) const;
 
 	/*
 	 * Attempt to begin the process of asynchronously creating a folder at the given virtual path, returning a temporary item that can be finalized or canceled by the user.
@@ -319,124 +321,124 @@ public:
 	 *
 	 * @return The pending folder item to create (test for validity).
 	 */
-	FContentBrowserItemTemporaryContext CreateFolder(const FName InPath, const TSharedPtr<IContentBrowserHideFolderIfEmptyFilter>& HideFolderIfEmptyFilter) const;
+	UE_API FContentBrowserItemTemporaryContext CreateFolder(const FName InPath, const TSharedPtr<IContentBrowserHideFolderIfEmptyFilter>& HideFolderIfEmptyFilter) const;
 
 	/**
 	 * Attempt to convert the given package path to virtual paths associated with the active data sources (callback will be called for each successful conversion).
 	 * @note This exists to allow the Content Browser to interface with public APIs that only operate on package paths and should ideally be avoided for new code.
 	 * @note This function only adjusts the path to something that could represent a virtualized item within this data source, but it doesn't guarantee that an item actually exists at that path.
 	 */
-	void Legacy_TryConvertPackagePathToVirtualPaths(const FName InPackagePath, TFunctionRef<bool(FName)> InCallback);
+	UE_API void Legacy_TryConvertPackagePathToVirtualPaths(const FName InPackagePath, TFunctionRef<bool(FName)> InCallback);
 
 	/**
 	 * Attempt to convert the given asset data to a virtual paths associated with the active data sources (callback will be called for each successful conversion).
 	 * @note This exists to allow the Content Browser to interface with public APIs that only operate on asset data and should ideally be avoided for new code.
 	 * @note This function only adjusts the path to something that could represent a virtualized item within this data source, but it doesn't guarantee that an item actually exists at that path.
 	 */
-	void Legacy_TryConvertAssetDataToVirtualPaths(const FAssetData& InAssetData, const bool InUseFolderPaths, TFunctionRef<bool(FName)> InCallback);
+	UE_API void Legacy_TryConvertAssetDataToVirtualPaths(const FAssetData& InAssetData, const bool InUseFolderPaths, TFunctionRef<bool(FName)> InCallback);
 
 	/**
 	 * Rebuild the virtual path tree if rules have changed.
 	 */
-	void RefreshVirtualPathTreeIfNeeded();
+	UE_API void RefreshVirtualPathTreeIfNeeded();
 
 	/**
 	 * Call when rules of virtual path generation have changed beyond content browser settings.
 	 */
-	void SetVirtualPathTreeNeedsRebuild();
+	UE_API void SetVirtualPathTreeNeedsRebuild();
 
 	/**
 	 * Converts an internal path to a virtual path based on current rules
 	 */
-	void ConvertInternalPathToVirtual(const FStringView InPath, FName& OutPath);
-	void ConvertInternalPathToVirtual(FName InPath, FName& OutPath);
-	FName ConvertInternalPathToVirtual(FName InPath);
-	TArray<FString> ConvertInternalPathsToVirtual(const TArray<FString>& InPaths);
+	UE_API void ConvertInternalPathToVirtual(const FStringView InPath, FName& OutPath);
+	UE_API void ConvertInternalPathToVirtual(FName InPath, FName& OutPath);
+	UE_API FName ConvertInternalPathToVirtual(FName InPath);
+	UE_API TArray<FString> ConvertInternalPathsToVirtual(const TArray<FString>& InPaths);
 
 	/**
 	 * Converts a virtual path to a display path, where each path segment is replaced with the display name of its corresponding FContentBrowserItem.
 	 * This path will display similar to the Content Browser's navigation bar.
 	 */
-	FText ConvertVirtualPathToDisplay(const FStringView InVirtualPath, const EContentBrowserItemTypeFilter InLeafType) const;
-	FText ConvertVirtualPathToDisplay(FName InVirtualPath, const EContentBrowserItemTypeFilter InLeafType) const;
-	FText ConvertVirtualPathToDisplay(const FContentBrowserItem& InItem) const;
+	UE_API FText ConvertVirtualPathToDisplay(const FStringView InVirtualPath, const EContentBrowserItemTypeFilter InLeafType) const;
+	UE_API FText ConvertVirtualPathToDisplay(FName InVirtualPath, const EContentBrowserItemTypeFilter InLeafType) const;
+	UE_API FText ConvertVirtualPathToDisplay(const FContentBrowserItem& InItem) const;
 
 	/**
 	 * Converts virtual path back into an internal or invariant path
 	 */
-	EContentBrowserPathType TryConvertVirtualPath(const FStringView InPath, FStringBuilderBase& OutPath) const;
-	EContentBrowserPathType TryConvertVirtualPath(const FStringView InPath, FString& OutPath) const;
-	EContentBrowserPathType TryConvertVirtualPath(const FStringView InPath, FName& OutPath) const;
-	EContentBrowserPathType TryConvertVirtualPath(const FName InPath, FName& OutPath) const;
+	UE_API EContentBrowserPathType TryConvertVirtualPath(const FStringView InPath, FStringBuilderBase& OutPath) const;
+	UE_API EContentBrowserPathType TryConvertVirtualPath(const FStringView InPath, FString& OutPath) const;
+	UE_API EContentBrowserPathType TryConvertVirtualPath(const FStringView InPath, FName& OutPath) const;
+	UE_API EContentBrowserPathType TryConvertVirtualPath(const FName InPath, FName& OutPath) const;
 
 	/**
 	 * Returns array of paths converted to internal.
 	 */
-	TArray<FString> TryConvertVirtualPathsToInternal(const TArray<FString>& InVirtualPaths) const;
+	UE_API TArray<FString> TryConvertVirtualPathsToInternal(const TArray<FString>& InVirtualPaths) const;
 
 	/**
 	 * Customize list of folders that appear first in content browser based on internal or invariant paths
 	 */
-	void SetPathViewSpecialSortFolders(const TArray<FName>& InSpecialSortFolders);
+	UE_API void SetPathViewSpecialSortFolders(const TArray<FName>& InSpecialSortFolders);
 
 	/**
 	 * Returns reference to list of paths that appear first in content browser based on internal or invariant paths
 	 */
-	const TArray<FName>& GetPathViewSpecialSortFolders() const;
+	UE_API const TArray<FName>& GetPathViewSpecialSortFolders() const;
 
 	/**
 	 * Returns reference to default list of paths that appear first in content browser based on internal or invariant paths
 	 */
-	const TArray<FName>& GetDefaultPathViewSpecialSortFolders() const;
+	UE_API const TArray<FName>& GetDefaultPathViewSpecialSortFolders() const;
 
 	/**
 	 * Set delegate used to generate a virtual path.
 	 */
-	void SetGenerateVirtualPathPrefixDelegate(const FContentBrowserGenerateVirtualPathDelegate& InDelegate);
+	UE_API void SetGenerateVirtualPathPrefixDelegate(const FContentBrowserGenerateVirtualPathDelegate& InDelegate);
 
 	/**
 	 * Delegate called to generate a virtual path. Can be set to override default behavior.
 	 */
-	FContentBrowserGenerateVirtualPathDelegate& OnGenerateVirtualPathPrefix();
+	UE_API FContentBrowserGenerateVirtualPathDelegate& OnGenerateVirtualPathPrefix();
 
 	/**
 	 * Return whether the plugin VersePath should be its path in the content browser, if it has one
 	 */
-	bool UsePluginVersePath(const TSharedRef<IPlugin>& Plugin);
+	UE_API bool UsePluginVersePath(const TSharedRef<IPlugin>& Plugin);
 
 	/**
 	 * Set delegate to check whether the plugin VersePath should be its path in the content browser, if it has one
 	 */
-	void SetUsePluginVersePathDelegate(FContentBrowserUsePluginVersePathDelegate InDelegate);
+	UE_API void SetUsePluginVersePathDelegate(FContentBrowserUsePluginVersePathDelegate InDelegate);
 
 	/**
 	 * Delegate to check whether the plugin VersePath should be its path in the content browser, if it has one
 	 */
-	FContentBrowserUsePluginVersePathDelegate& GetUsePluginVersePathDelegate();
+	UE_API FContentBrowserUsePluginVersePathDelegate& GetUsePluginVersePathDelegate();
 
 	/**
 	 * Creates the default IContentBrowserHideFolderIfEmptyFilter, and combines it with the results of any additional registered FContentBrowserCreateHideFolderIfEmptyFilter delegates.
 	 * @note The filter is only valid until any dependent data changes, so only keep it for a short time (typically within a function call tree).
 	 */
-	TSharedPtr<IContentBrowserHideFolderIfEmptyFilter> CreateHideFolderIfEmptyFilter() const;
+	UE_API TSharedPtr<IContentBrowserHideFolderIfEmptyFilter> CreateHideFolderIfEmptyFilter() const;
 
 	/**
 	* Registers/deregisters a delegate used to create a filter used to filter out empty folders. If multiple delegates are registered they will get or'd togeather.
 	* Intended for content folders that always exist but do not generally contain assets, such as /Game/Collections
 	*/
-	FDelegateHandle RegisterCreateHideFolderIfEmptyFilter(FContentBrowserCreateHideFolderIfEmptyFilter Delegate);
-	void UnregisterCreateHideFolderIfEmptyFilter(FDelegateHandle DelegateHandle);
+	UE_API FDelegateHandle RegisterCreateHideFolderIfEmptyFilter(FContentBrowserCreateHideFolderIfEmptyFilter Delegate);
+	UE_API void UnregisterCreateHideFolderIfEmptyFilter(FDelegateHandle DelegateHandle);
 
 	/**
 	 * Prefix to use when generating virtual paths and "Show All Folder" option is enabled
 	 */
-	const FString& GetAllFolderPrefix() const;
+	UE_API const FString& GetAllFolderPrefix() const;
 
 	/**
 	 * Permission list that controls whether content in a given folder path can be edited.
 	 * Note: This does not control if the folder path is writable or if content can be deleted.
 	 */
-	TSharedRef<FPathPermissionList>& GetEditableFolderPermissionList();
+	UE_API TSharedRef<FPathPermissionList>& GetEditableFolderPermissionList();
 
 
 	/**
@@ -455,11 +457,11 @@ public:
 
 private:
 
-	void InitializeCacheIDOwner(FContentBrowserDataFilterCacheIDOwner& IDOwner);
+	UE_API void InitializeCacheIDOwner(FContentBrowserDataFilterCacheIDOwner& IDOwner);
 
-	void RemoveUnusedCachedFilterData(const FContentBrowserDataFilterCacheIDOwner& IDOwner, TArrayView<const FName> InVirtualPathsInUse, const FContentBrowserDataFilter& DataFilter) const;
+	UE_API void RemoveUnusedCachedFilterData(const FContentBrowserDataFilterCacheIDOwner& IDOwner, TArrayView<const FName> InVirtualPathsInUse, const FContentBrowserDataFilter& DataFilter) const;
 
-	void ClearCachedFilterData(const FContentBrowserDataFilterCacheIDOwner& IDOwner) const;
+	UE_API void ClearCachedFilterData(const FContentBrowserDataFilterCacheIDOwner& IDOwner) const;
 
 	using FNameToDataSourceMap = TSortedMap<FName, UContentBrowserDataSource*, FDefaultAllocator, FNameFastLess>;
 
@@ -467,46 +469,46 @@ private:
 	 * Called to handle a data source modular feature being registered.
 	 * @note Will activate the data source if it is in the EnabledDataSources array.
 	 */
-	void HandleDataSourceRegistered(const FName& Type, class IModularFeature* Feature);
+	UE_API void HandleDataSourceRegistered(const FName& Type, class IModularFeature* Feature);
 	
 	/**
 	 * Called to handle a data source modular feature being unregistered.
 	 * @note Will deactivate the data source if it is in the ActiveDataSources map.
 	 */
-	void HandleDataSourceUnregistered(const FName& Type, class IModularFeature* Feature);
+	UE_API void HandleDataSourceUnregistered(const FName& Type, class IModularFeature* Feature);
 
 	/**
 	 * Tick this subsystem.
 	 * @note Called once every 0.1 seconds.
 	 */
-	void Tick(const float InDeltaTime);
+	UE_API void Tick(const float InDeltaTime);
 
 	/**
 	 * Returns true if item data modifications are being processed.
 	 */
-	bool AllowModifiedItemDataUpdates() const;
+	UE_API bool AllowModifiedItemDataUpdates() const;
 
 	/**
 	 * Called when Play in Editor begins.
 	 */
-	void OnBeginPIE(const bool bIsSimulating);
+	UE_API void OnBeginPIE(const bool bIsSimulating);
 
 	/**
 	 * Called when Play in Editor stops.
 	 */
-	void OnEndPIE(const bool bIsSimulating);
+	UE_API void OnEndPIE(const bool bIsSimulating);
 
 	/**
 	 * Called when Content added to delay content browser tick for a frame
 	 * Prevents content browser from slowing down initialization by ticking as content is loaded
 	 */
-	void OnContentPathMounted(const FString& AssetPath, const FString& ContentPath);
+	UE_API void OnContentPathMounted(const FString& AssetPath, const FString& ContentPath);
 
 	//~ IContentBrowserItemDataSink interface
-	virtual void QueueItemDataUpdate(FContentBrowserItemDataUpdate&& InUpdate) override;
+	UE_API virtual void QueueItemDataUpdate(FContentBrowserItemDataUpdate&& InUpdate) override;
 
 	UE_DEPRECATED(5.5, "NotifyItemDataRefreshed is deprecated, for editor performance reasons no external systems should be able to request a full refresh.")
-	virtual void NotifyItemDataRefreshed() override;
+	UE_API virtual void NotifyItemDataRefreshed() override;
 
 	/**
 	 * Handle for the Tick callback.
@@ -650,3 +652,5 @@ public:
 private:
 	UContentBrowserDataSubsystem* ContentBrowserData = nullptr;
 };
+
+#undef UE_API
diff --git a/Engine/Source/Editor/ContentBrowserData/Public/ContentBrowserItemPath.h b/Engine/Source/Editor/ContentBrowserData/Public/ContentBrowserItemPath.h
--- a/Engine/Source/Editor/ContentBrowserData/Public/ContentBrowserItemPath.h
+++ b/Engine/Source/Editor/ContentBrowserData/Public/ContentBrowserItemPath.h
@@ -16,6 +16,8 @@
 
 #include "ContentBrowserItemPath.generated.h"
 
+#define UE_API CONTENTBROWSERDATA_API
+
 class UObject;
 struct FFrame;
 
@@ -25,11 +27,11 @@ struct FFrame;
  * Path conversion each time Set is called
  */
 USTRUCT(BlueprintType, meta=(HasNativeBreak="/Script/ContentBrowserData.ContentBrowserItemPathExtensions.BreakContentBrowserItemPath", HasNativeMake="/Script/ContentBrowserData.ContentBrowserItemPathExtensions.MakeContentBrowserItemPath"))
-struct CONTENTBROWSERDATA_API FContentBrowserItemPath
+struct FContentBrowserItemPath
 {
 	GENERATED_BODY()
 
-	FContentBrowserItemPath();
+	UE_API FContentBrowserItemPath();
 
 	FContentBrowserItemPath(const FName InVirtualPath, const FName InInternalPath)
 		: VirtualPath(InVirtualPath)
@@ -39,9 +41,9 @@ struct CONTENTBROWSERDATA_API FContentBrowserItemPath
 		check(!InternalPath.IsNone());
 	}
 
-	FContentBrowserItemPath(const FStringView InPath, const EContentBrowserPathType InPathType);
-	FContentBrowserItemPath(const TCHAR* InPath, const EContentBrowserPathType InPathType);
-	FContentBrowserItemPath(const FName InPath, const EContentBrowserPathType InPathType);
+	UE_API FContentBrowserItemPath(const FStringView InPath, const EContentBrowserPathType InPathType);
+	UE_API FContentBrowserItemPath(const TCHAR* InPath, const EContentBrowserPathType InPathType);
+	UE_API FContentBrowserItemPath(const FName InPath, const EContentBrowserPathType InPathType);
 
 	bool operator==(const FContentBrowserItemPath& Other) const 
 	{ 
@@ -51,37 +53,37 @@ struct CONTENTBROWSERDATA_API FContentBrowserItemPath
 	/**
 	 * Set the path being stored
 	 */
-	void SetPathFromString(const FStringView InPath, const EContentBrowserPathType InPathType);
+	UE_API void SetPathFromString(const FStringView InPath, const EContentBrowserPathType InPathType);
 
 	/**
 	 * Set the path being stored
 	 */
-	void SetPathFromName(const FName InPath, const EContentBrowserPathType InPathType);
+	UE_API void SetPathFromName(const FName InPath, const EContentBrowserPathType InPathType);
 
 	/**
 	 * Returns virtual path as FName (eg, "/All/Plugins/PluginA/MyFile").
 	 */
-	FName GetVirtualPathName() const;
+	UE_API FName GetVirtualPathName() const;
 
 	/**
 	 * Returns internal path if there is one (eg,. "/PluginA/MyFile").
 	 */
-	FName GetInternalPathName() const;
+	UE_API FName GetInternalPathName() const;
 
 	/**
 	 * Returns virtual path as newly allocated FString (eg, "/All/Plugins/PluginA/MyFile").
 	 */
-	FString GetVirtualPathString() const;
+	UE_API FString GetVirtualPathString() const;
 
 	/**
 	 * Returns internal path as newly allocated FString if there is one or empty FString (eg,. "/PluginA/MyFile").
 	 */
-	FString GetInternalPathString() const;
+	UE_API FString GetInternalPathString() const;
 
 	/**
 	 * Returns true if there is an internal path
 	 */
-	bool HasInternalPath() const;
+	UE_API bool HasInternalPath() const;
 
 private:
 	/** Path as virtual (eg, "/All/Plugins/PluginA/MyFile") */
@@ -138,3 +140,5 @@ public:
 		InternalPath = ItemPath.HasInternalPath() ? ItemPath.GetInternalPathName() : NAME_None;
 	}
  };
+
+#undef UE_API
diff --git a/Engine/Source/Editor/ContentBrowserData/Public/ContentBrowserVirtualPathTree.h b/Engine/Source/Editor/ContentBrowserData/Public/ContentBrowserVirtualPathTree.h
--- a/Engine/Source/Editor/ContentBrowserData/Public/ContentBrowserVirtualPathTree.h
+++ b/Engine/Source/Editor/ContentBrowserData/Public/ContentBrowserVirtualPathTree.h
@@ -11,23 +11,25 @@
 #include "Misc/StringBuilder.h"
 #include "UObject/NameTypes.h"
 
+#define UE_API CONTENTBROWSERDATA_API
+
 class FString;
 template <typename FuncType> class TFunctionRef;
 
 /**
  * Tree of virtual paths ending where internal paths start. Used for conversion of paths and during enumerate.
  */
-class CONTENTBROWSERDATA_API FContentBrowserVirtualPathTree
+class FContentBrowserVirtualPathTree
 {
 public:
 	/** Adds the specified path to the tree, creating nodes as needed and calling OnPathAdded for any new paths added. Returns true if the specified path was actually added (as opposed to already existed) */
-	bool CachePath(FName Path, FName InternalPath, TFunctionRef<void(FName)> OnPathAdded);
+	UE_API bool CachePath(FName Path, FName InternalPath, TFunctionRef<void(FName)> OnPathAdded);
 
 	/** Removes the specified path from the tree, calling OnPathRemoved for any existing paths removed. Returns true if the path was found and removed. */
-	bool RemovePath(FName Path, TFunctionRef<void(FName)> OnPathRemoved);
+	UE_API bool RemovePath(FName Path, TFunctionRef<void(FName)> OnPathRemoved);
 
 	/** Checks whether the given path is one that we know about */
-	bool PathExists(FName Path, bool& bIsFullyVirtual) const;
+	UE_API bool PathExists(FName Path, bool& bIsFullyVirtual) const;
 
 	/** Checks whether the given path is one that we know about */
 	bool PathExists(FName Path) const
@@ -37,16 +39,16 @@ public:
 	}
 
 	/** Get all of the paths we know about */
-	bool GetAllPaths(TSet<FName>& OutPaths) const;
+	UE_API bool GetAllPaths(TSet<FName>& OutPaths) const;
 
 	/** Enumerate all of the paths we know about */
-	void EnumerateAllPaths(TFunctionRef<bool(FName)> Callback) const;
+	UE_API void EnumerateAllPaths(TFunctionRef<bool(FName)> Callback) const;
 
 	/** Recursively gathers all child paths from the specified base path relative to this node */
-	bool GetSubPaths(FName BasePath, TSet<FName>& OutPaths, bool bRecurse = true) const;
+	UE_API bool GetSubPaths(FName BasePath, TSet<FName>& OutPaths, bool bRecurse = true) const;
 
 	/** Recursively enumerates all child paths from the specified base path relative to this node */
-	bool EnumerateSubPaths(FName BasePath, TFunctionRef<bool(FName, FName)> Callback, bool bRecurse = true) const;
+	UE_API bool EnumerateSubPaths(FName BasePath, TFunctionRef<bool(FName, FName)> Callback, bool bRecurse = true) const;
 
 	SIZE_T GetAllocatedSize(void) const
 	{
@@ -61,13 +63,13 @@ public:
 	}
 
 	/** Clears all paths */
-	void Reset();
+	UE_API void Reset();
 
 	/** Returns name of parent virtual path when given a virtual path */
-	FName GetParentPath(FName Path) const;
+	UE_API FName GetParentPath(FName Path) const;
 
 	/** Returns number of paths in total */
-	int32 NumPaths() const;
+	UE_API int32 NumPaths() const;
 
 	/** Returns reference to map that allows looking up an internal path when given a virtual path */
 	const TMap<FName, FName>& GetVirtualToInternalMounts() const
@@ -81,7 +83,7 @@ public:
 	 * Removes starting forward slash and Classes_ prefix 
 	 * Example: "/Classes_A/Textures" returns "A" and sets bOutHadClassesPrefix=true
 	 */
-	static FStringView GetMountPointFromPath(const FStringView InPath, bool& bOutHadClassesPrefix);
+	static UE_API FStringView GetMountPointFromPath(const FStringView InPath, bool& bOutHadClassesPrefix);
 
 	/** Tries to converts virtual path into an internal path.
 	 * Returns None if virtual prefix portion of the path does not exist in this path tree
@@ -89,17 +91,17 @@ public:
 	 * Returns Internal if virtual prefix portion of the path exists in this path tree and there is text afterwards
 	 * @Note: When Internal is the return value it does not check to see if the internal path actually exists or not
 	 */
-	EContentBrowserPathType TryConvertVirtualPathToInternal(FStringView InPath, FStringBuilderBase& OutPath) const;
-	EContentBrowserPathType TryConvertVirtualPathToInternal(FStringView InPath, FString& OutPath) const;
-	EContentBrowserPathType TryConvertVirtualPathToInternal(FName InPath, FName& OutPath) const;
+	UE_API EContentBrowserPathType TryConvertVirtualPathToInternal(FStringView InPath, FStringBuilderBase& OutPath) const;
+	UE_API EContentBrowserPathType TryConvertVirtualPathToInternal(FStringView InPath, FString& OutPath) const;
+	UE_API EContentBrowserPathType TryConvertVirtualPathToInternal(FName InPath, FName& OutPath) const;
 
 private:
 
 	/** Is the user option to show the all folder enabled */
-	bool IsShowAllFolderEnabled() const;
+	UE_API bool IsShowAllFolderEnabled() const;
 
 	/** Get the string to use when show all folder is enabled */
-	const FString& GetAllFolderPrefix() const;
+	UE_API const FString& GetAllFolderPrefix() const;
 
 	/** A one-to-one mapping between a virtual path and internal path */
 	TMap<FName, FName> VirtualToInternalMounts;
@@ -110,3 +112,5 @@ private:
 	/** A one-to-one mapping between a child path and its parent path. Paths without a parent (root paths) will not appear in this map. */
 	TMap<FName, FName> ChildPathToParentPath;
 };
+
+#undef UE_API
diff --git a/Engine/Source/Editor/CurveEditor/Private/Modification/CurveEditorTransactionObject.h b/Engine/Source/Editor/CurveEditor/Private/Modification/CurveEditorTransactionObject.h
--- a/Engine/Source/Editor/CurveEditor/Private/Modification/CurveEditorTransactionObject.h
+++ b/Engine/Source/Editor/CurveEditor/Private/Modification/CurveEditorTransactionObject.h
@@ -14,8 +14,8 @@ class FCurveEditor;
  * By design, the transaction system requires a UObject to associate FCommandChanges with. This is usually an asset.
  * Conceptually, instances of this class will act as an "asset" and the associated FCurveEditor contains the "asset data", i.e. its  FCurveModels.
  */
-UCLASS()
-class CURVEEDITOR_API UCurveEditorTransactionObject : public UObject
+UCLASS(MinimalAPI)
+class UCurveEditorTransactionObject : public UObject
 {
 	GENERATED_BODY()
 public:
diff --git a/Engine/Source/Editor/CurveEditor/Private/SCurveEditorViewContainer.h b/Engine/Source/Editor/CurveEditor/Private/SCurveEditorViewContainer.h
--- a/Engine/Source/Editor/CurveEditor/Private/SCurveEditorViewContainer.h
+++ b/Engine/Source/Editor/CurveEditor/Private/SCurveEditorViewContainer.h
@@ -15,6 +15,8 @@
 #include "Widgets/DeclarativeSyntaxSupport.h"
 #include "Widgets/SBoxPanel.h"
 
+#define UE_API CURVEEDITOR_API
+
 class FCurveEditor;
 class FPaintArgs;
 class FSlateRect;
@@ -32,7 +34,7 @@ struct FPointerEvent;
 /**
  * Curve editor widget that reflects the state of an FCurveEditor
  */
-class CURVEEDITOR_API SCurveEditorViewContainer : public SVerticalBox
+class SCurveEditorViewContainer : public SVerticalBox
 {
 	SLATE_BEGIN_ARGS(SCurveEditorViewContainer)
 		: _MinimumPanelHeight(300.0f)
@@ -48,46 +50,46 @@ class CURVEEDITOR_API SCurveEditorViewContainer : public SVerticalBox
 	/**
 	 * Construct a new curve editor panel widget
 	 */
-	void Construct(const FArguments& InArgs, TSharedRef<FCurveEditor> InCurveEditor);
+	UE_API void Construct(const FArguments& InArgs, TSharedRef<FCurveEditor> InCurveEditor);
 
 	TArrayView<const TSharedPtr<SCurveEditorView>> GetViews() const { return Views; }
 
-	void AddView(TSharedRef<SCurveEditorView> ViewToAdd);
+	UE_API void AddView(TSharedRef<SCurveEditorView> ViewToAdd);
 
-	void Clear();
+	UE_API void Clear();
 
 private:
 
 	// SWidget Interface
-	virtual int32 OnPaint(const FPaintArgs& Args, const FGeometry& AllottedGeometry, const FSlateRect& MyCullingRect, FSlateWindowElementList& OutDrawElements, int32 LayerId, const FWidgetStyle& InWidgetStyle, bool bParentEnabled) const override;
-	virtual void Tick(const FGeometry& AllottedGeometry, const double InCurrentTime, const float InDeltaTime) override;
-	virtual FVector2D ComputeDesiredSize(float) const;
-	virtual bool ComputeVolatility() const override;
+	UE_API virtual int32 OnPaint(const FPaintArgs& Args, const FGeometry& AllottedGeometry, const FSlateRect& MyCullingRect, FSlateWindowElementList& OutDrawElements, int32 LayerId, const FWidgetStyle& InWidgetStyle, bool bParentEnabled) const override;
+	UE_API virtual void Tick(const FGeometry& AllottedGeometry, const double InCurrentTime, const float InDeltaTime) override;
+	UE_API virtual FVector2D ComputeDesiredSize(float) const;
+	UE_API virtual bool ComputeVolatility() const override;
 	// ~SWidget Interface
 
 	/*~ Mouse interaction */
-	virtual FReply OnPreviewMouseButtonDown(const FGeometry& MyGeometry, const FPointerEvent& MouseEvent) override;
-	virtual FReply OnMouseButtonDown(const FGeometry& MyGeometry, const FPointerEvent& MouseEvent) override;
-	virtual FReply OnMouseButtonUp(const FGeometry& MyGeometry, const FPointerEvent& MouseEvent) override;
-	virtual FReply OnMouseMove(const FGeometry& MyGeometry, const FPointerEvent& MouseEvent) override;
-	virtual void OnFinishedPointerInput() override;
-	virtual FReply OnMouseButtonDoubleClick(const FGeometry& InMyGeometry, const FPointerEvent& InMouseEvent) override;
-	virtual void OnFocusLost(const FFocusEvent& InFocusEvent) override;
-	virtual FCursorReply OnCursorQuery(const FGeometry& MyGeometry, const FPointerEvent& CursorEvent) const override;
+	UE_API virtual FReply OnPreviewMouseButtonDown(const FGeometry& MyGeometry, const FPointerEvent& MouseEvent) override;
+	UE_API virtual FReply OnMouseButtonDown(const FGeometry& MyGeometry, const FPointerEvent& MouseEvent) override;
+	UE_API virtual FReply OnMouseButtonUp(const FGeometry& MyGeometry, const FPointerEvent& MouseEvent) override;
+	UE_API virtual FReply OnMouseMove(const FGeometry& MyGeometry, const FPointerEvent& MouseEvent) override;
+	UE_API virtual void OnFinishedPointerInput() override;
+	UE_API virtual FReply OnMouseButtonDoubleClick(const FGeometry& InMyGeometry, const FPointerEvent& InMouseEvent) override;
+	UE_API virtual void OnFocusLost(const FFocusEvent& InFocusEvent) override;
+	UE_API virtual FCursorReply OnCursorQuery(const FGeometry& MyGeometry, const FPointerEvent& CursorEvent) const override;
 
 	virtual bool SupportsKeyboardFocus() const override { return true; }
-	virtual FReply OnKeyDown(const FGeometry& MyGeometry, const FKeyEvent& InKeyEvent) override;
-	virtual FReply OnKeyUp(const FGeometry& MyGeometry, const FKeyEvent& InKeyEvent) override;
+	UE_API virtual FReply OnKeyDown(const FGeometry& MyGeometry, const FKeyEvent& InKeyEvent) override;
+	UE_API virtual FReply OnKeyUp(const FGeometry& MyGeometry, const FKeyEvent& InKeyEvent) override;
 
 private:
 
-	FMargin GetSlotPadding(int32 SlotIndex) const;
+	UE_API FMargin GetSlotPadding(int32 SlotIndex) const;
 
-	void OnCurveEditorToolChanged(FCurveEditorToolID InToolId);
+	UE_API void OnCurveEditorToolChanged(FCurveEditorToolID InToolId);
 
-	void ExpandInputBounds(float NewWidth);
+	UE_API void ExpandInputBounds(float NewWidth);
 
-	bool IsScrubTimeKeyEvent(const FKeyEvent& InKeyEvent);
+	UE_API bool IsScrubTimeKeyEvent(const FKeyEvent& InKeyEvent);
 
 private:
 
@@ -119,4 +121,6 @@ private:
 
 	/** Whether or not we are scrubbing time*/
 	bool bIsScrubbingTime = false;
-};
\ No newline at end of file
+};
+
+#undef UE_API
diff --git a/Engine/Source/Editor/CurveEditor/Public/CurveDrawInfo.h b/Engine/Source/Editor/CurveEditor/Public/CurveDrawInfo.h
--- a/Engine/Source/Editor/CurveEditor/Public/CurveDrawInfo.h
+++ b/Engine/Source/Editor/CurveEditor/Public/CurveDrawInfo.h
@@ -14,7 +14,7 @@ struct FSlateBrush;
 /**
  * Structure that defines how to draw a particular key
  */
-struct CURVEEDITOR_API FKeyDrawInfo
+struct FKeyDrawInfo
 {
 	FKeyDrawInfo()
 		: ScreenSize(0.f, 0.f)
@@ -36,7 +36,7 @@ struct CURVEEDITOR_API FKeyDrawInfo
 /**
  * Structure that defines the necessary data for painting a given curve point
  */
-struct CURVEEDITOR_API FCurvePointInfo
+struct FCurvePointInfo
 {
 	FCurvePointInfo(FKeyHandle InHandle)
 		: KeyHandle(InHandle)
@@ -66,7 +66,7 @@ struct CURVEEDITOR_API FCurvePointInfo
 /**
  * Structure that defines the necessary data for painting a whole curve
  */
-struct CURVEEDITOR_API FCurveDrawParams
+struct FCurveDrawParams
 {
 	/** The color to draw this curve */
 	FLinearColor Color;
diff --git a/Engine/Source/Editor/CurveEditor/Public/CurveEditorCurveDrawParamsCache.h b/Engine/Source/Editor/CurveEditor/Public/CurveEditorCurveDrawParamsCache.h
--- a/Engine/Source/Editor/CurveEditor/Public/CurveEditorCurveDrawParamsCache.h
+++ b/Engine/Source/Editor/CurveEditor/Public/CurveEditorCurveDrawParamsCache.h
@@ -6,6 +6,8 @@
 #include "CurveEditorSettings.h"
 #include "Templates/SharedPointer.h"
 
+#define UE_API CURVEEDITOR_API
+
 class FCurveEditor;
 class FCurveModel;
 class SCurveEditorView;
@@ -19,7 +21,7 @@ namespace UE::CurveEditor
 	struct FCurveDrawParamsHandle;
 
 	/** Caches curve draw params for curves of a curve editor. */
-	class CURVEEDITOR_API FCurveDrawParamsCache
+	class FCurveDrawParamsCache
 		: public TSharedFromThis<FCurveDrawParamsCache>
 	{
 		// Let SCurveEditorView write directly to the array of cached curve params when CVar CurveEditor.UseCurveCache is disabled
@@ -64,13 +66,13 @@ namespace UE::CurveEditor
 			FVector2D CachedGeometrySize = FVector2D::ZeroVector;
 		};
 
-		FCurveDrawParamsCache();
+		UE_API FCurveDrawParamsCache();
 
 		/** 
 		 * Invalidates the cache for specified view and curve model IDs and causes an update on cached data.
 		 * The cache may perform updates that are not instantly relevant asynchronously. To update data synchronous use UpdateCurveDrawParamSynchonous instead.
 		 */
-		void Invalidate(const TSharedRef<SCurveEditorView>& CurveEditorView, const TArray<FCurveModelID>& ModelIDs = TArray<FCurveModelID>());
+		UE_API void Invalidate(const TSharedRef<SCurveEditorView>& CurveEditorView, const TArray<FCurveModelID>& ModelIDs = TArray<FCurveModelID>());
 
 		/**
 		 * Returns how the curve cache has changed since it was last generated.
@@ -85,20 +87,20 @@ namespace UE::CurveEditor
 		const TArray<FCurveDrawParams>& GetCurveDrawParams() const { return CachedDrawParams; }
 
 		/** Updates curve draw params for all specified curve models, synchronous. */
-		void UpdateAllCurveDrawParamSynchonous(const TSharedRef<SCurveEditorView>& CurveEditorView, const TArray<FCurveModelID>& CurveModelIDs, TArray<FCurveDrawParams>& OutParams);
+		UE_API void UpdateAllCurveDrawParamSynchonous(const TSharedRef<SCurveEditorView>& CurveEditorView, const TArray<FCurveModelID>& CurveModelIDs, TArray<FCurveDrawParams>& OutParams);
 
 		/** Gets curve draw params for a curve model, synchronous. */
-		void UpdateCurveDrawParamsSynchonous(const FCurveEditorScreenSpace& CurveSpace, FCurveModel* CurveModel, const FCurveModelID& ModelID, FCurveDrawParams& Params);
+		UE_API void UpdateCurveDrawParamsSynchonous(const FCurveEditorScreenSpace& CurveSpace, FCurveModel* CurveModel, const FCurveModelID& ModelID, FCurveDrawParams& Params);
 
 	private:
 		/** Updates the curve cache flags */
-		void UpdateCurveCacheFlags(const TSharedRef<SCurveEditorView>& CurveEditorView);
+		UE_API void UpdateCurveCacheFlags(const TSharedRef<SCurveEditorView>& CurveEditorView);
 
 		/** Draws the curves, possibly async */
-		void DrawCurves(const TSharedRef<SCurveEditorView>& CurveEditorView, const TArray<FCurveModelID>& ModelIDs = TArray<FCurveModelID>());
+		UE_API void DrawCurves(const TSharedRef<SCurveEditorView>& CurveEditorView, const TArray<FCurveModelID>& ModelIDs = TArray<FCurveModelID>());
 
 		/** Updates curve draw params, possibly async */
-		void UpdateCurveDrawParams(const TSharedRef<SCurveEditorView>& CurveEditorView, const TArray<FCurveDrawParamsHandle>& CurveDrawParamsHandles);
+		UE_API void UpdateCurveDrawParams(const TSharedRef<SCurveEditorView>& CurveEditorView, const TArray<FCurveDrawParamsHandle>& CurveDrawParamsHandles);
 
 		/** Curve cache flags that change based upon data or view getting modified*/
 		ECurveCacheFlags CurveCacheFlags = ECurveCacheFlags::UpdateAll;
@@ -113,3 +115,5 @@ namespace UE::CurveEditor
 		TWeakPtr<FCurveEditor> WeakCurveEditor;
 	};
 }
+
+#undef UE_API
diff --git a/Engine/Source/Editor/CurveEditor/Public/CurveEditorCurveDrawParamsHandle.h b/Engine/Source/Editor/CurveEditor/Public/CurveEditorCurveDrawParamsHandle.h
--- a/Engine/Source/Editor/CurveEditor/Public/CurveEditorCurveDrawParamsHandle.h
+++ b/Engine/Source/Editor/CurveEditor/Public/CurveEditorCurveDrawParamsHandle.h
@@ -4,6 +4,8 @@
 
 #include "CurveEditorTypes.h"
 
+#define UE_API CURVEEDITOR_API
+
 class FCurveDrawParamsCache;
 struct FCurveDrawParams;
 
@@ -12,18 +14,18 @@ namespace UE::CurveEditor
 	class FCurveDrawParamsCache;
 
 	/** Handle for an entry in the draw params array */
-	struct CURVEEDITOR_API FCurveDrawParamsHandle
+	struct FCurveDrawParamsHandle
 	{
 		FCurveDrawParamsHandle() = default;
 
 		/** Constructs the handle from the index in curve draw params. */
-		FCurveDrawParamsHandle(const TSharedRef<FCurveDrawParamsCache>& InDrawParamsCache, const int32 InIndex);
+		UE_API FCurveDrawParamsHandle(const TSharedRef<FCurveDrawParamsCache>& InDrawParamsCache, const int32 InIndex);
 
 		/** 
 		 * Returns the draw params or nullptr if the draw params no longer exists. 
 		 * Does a slow lookup if the array changed.
 		 */
-		FCurveDrawParams* Get() const;
+		UE_API FCurveDrawParams* Get() const;
 
 		/** Returns the curve model ID that corresponds to these curve draw params. */
 		const FCurveModelID& GetID() const { return CurveModelID; }
@@ -39,3 +41,5 @@ namespace UE::CurveEditor
 		TWeakPtr<class FCurveDrawParamsCache> WeakDrawParamsCache;
 	};
 }
+
+#undef UE_API
diff --git a/Engine/Source/Editor/CurveEditor/Public/CurveEditorKeyProxy.h b/Engine/Source/Editor/CurveEditor/Public/CurveEditorKeyProxy.h
--- a/Engine/Source/Editor/CurveEditor/Public/CurveEditorKeyProxy.h
+++ b/Engine/Source/Editor/CurveEditor/Public/CurveEditorKeyProxy.h
@@ -6,8 +6,8 @@
 #include "UObject/Interface.h"
 #include "CurveEditorKeyProxy.generated.h"
 
-UINTERFACE()
-class CURVEEDITOR_API UCurveEditorKeyProxy : public UInterface
+UINTERFACE(MinimalAPI)
+class UCurveEditorKeyProxy : public UInterface
 {
 	GENERATED_BODY()
 };
@@ -17,7 +17,7 @@ class CURVEEDITOR_API UCurveEditorKeyProxy : public UInterface
  * When used, UpdateValuesFromRawData will be called every frame to optionally retrieve the current values
  * of the key for this proxy.
  */
-class CURVEEDITOR_API ICurveEditorKeyProxy
+class ICurveEditorKeyProxy
 {
 public:
 
@@ -27,4 +27,4 @@ public:
 	 * Called by the curve editor to update this instance's properties with the underlying raw data, if necessary
 	 */
 	virtual void UpdateValuesFromRawData() = 0;
-};
\ No newline at end of file
+};
diff --git a/Engine/Source/Editor/CurveEditor/Public/CurveEditorSelection.h b/Engine/Source/Editor/CurveEditor/Public/CurveEditorSelection.h
--- a/Engine/Source/Editor/CurveEditor/Public/CurveEditorSelection.h
+++ b/Engine/Source/Editor/CurveEditor/Public/CurveEditorSelection.h
@@ -14,6 +14,8 @@
 #include "Misc/Optional.h"
 #include "Templates/SharedPointer.h"
 
+#define UE_API CURVEEDITOR_API
+
 class FCurveEditor;
 namespace UE::CurveEditor::KeySelection { class FAddInternal; }
 namespace UE::CurveEditor::KeySelection { class FSetSerialNumber; }
@@ -24,27 +26,27 @@ struct FCurvePointHandle;
  * A set of key handles implemented as a sorted array for transparent passing to TArrayView<> APIs.
  * Lookup is achieved via binary search: O(log(n)).
  */
-struct CURVEEDITOR_API FKeyHandleSet
+struct FKeyHandleSet
 {
 	/**
 	 * Add a new key handle to this set
 	 */
-	void Add(FKeyHandle Handle, ECurvePointType PointType);
+	UE_API void Add(FKeyHandle Handle, ECurvePointType PointType);
 
 	/**
 	 * Remove a handle from this set if it already exists, otherwise add it to the set
 	 */
-	void Toggle(FKeyHandle Handle, ECurvePointType PointType);
+	UE_API void Toggle(FKeyHandle Handle, ECurvePointType PointType);
 
 	/**
 	 * Remove a handle from this set
 	 */
-	void Remove(FKeyHandle Handle, ECurvePointType PointType);
+	UE_API void Remove(FKeyHandle Handle, ECurvePointType PointType);
 
 	/**
 	 * Check whether the specified handle exists in this set
 	 */
-	bool Contains(FKeyHandle Handle, ECurvePointType PointType) const;
+	UE_API bool Contains(FKeyHandle Handle, ECurvePointType PointType) const;
 
 	/**
 	 * Retrieve the number of handles in this set
@@ -85,7 +87,7 @@ private:
  * Class responsible for tracking selections of keys.
  * Only one type of point selection is supported at a time (key, arrive tangent, or leave tangent)
  */
-struct CURVEEDITOR_API FCurveEditorSelection
+struct FCurveEditorSelection
 {
 	friend UE::CurveEditor::KeySelection::FAddInternal;
 	friend UE::CurveEditor::KeySelection::FSetSerialNumber;
@@ -93,13 +95,13 @@ struct CURVEEDITOR_API FCurveEditorSelection
 	/**
 	 * Default constructor
 	 */
-	FCurveEditorSelection();
+	UE_API FCurveEditorSelection();
 
 	/**
 	 * Constructor which takes a reference to the curve editor, 
 	 * which is used to find if a model is read only
 	 */
-	FCurveEditorSelection(TWeakPtr<FCurveEditor> InWeakCurveEditor);
+	UE_API FCurveEditorSelection(TWeakPtr<FCurveEditor> InWeakCurveEditor);
 
 	/**
 	 * Retrieve this selection's serial number. Incremented whenever a change is made to the selection.
@@ -119,87 +121,87 @@ struct CURVEEDITOR_API FCurveEditorSelection
 	/**
 	 * Retrieve a set of selected key handles for the specified curve
 	 */
-	const FKeyHandleSet* FindForCurve(FCurveModelID InCurveID) const;
+	UE_API const FKeyHandleSet* FindForCurve(FCurveModelID InCurveID) const;
 
 	/**
 	 * Count the total number of selected keys by accumulating the number of selected keys for each curve
 	 */
-	int32 Count() const;
+	UE_API int32 Count() const;
 
 	/**
 	 * Check whether the specified handle is selected
 	 */
-	bool IsSelected(FCurvePointHandle InHandle) const;
+	UE_API bool IsSelected(FCurvePointHandle InHandle) const;
 
 	/**
 	 * Check whether the specified handle and curve ID is contained in this selection.
 	 */
-	bool Contains(FCurveModelID CurveID, FKeyHandle KeyHandle, ECurvePointType PointType) const;
+	UE_API bool Contains(FCurveModelID CurveID, FKeyHandle KeyHandle, ECurvePointType PointType) const;
 
 public:
 
 	/**
 	 * Add a point handle to this selection, changing the selection type if necessary.
 	 */
-	void Add(FCurvePointHandle InHandle);
+	UE_API void Add(FCurvePointHandle InHandle);
 
 	/**
 	 * Add a key handle to this selection, changing the selection type if necessary.
 	 */
-	void Add(FCurveModelID CurveID, ECurvePointType PointType, FKeyHandle KeyHandle);
+	UE_API void Add(FCurveModelID CurveID, ECurvePointType PointType, FKeyHandle KeyHandle);
 
 	/**
 	 * Add key handles to this selection, changing the selection type if necessary.
 	 */
-	void Add(FCurveModelID CurveID, ECurvePointType PointType, TArrayView<const FKeyHandle> Keys);
+	UE_API void Add(FCurveModelID CurveID, ECurvePointType PointType, TArrayView<const FKeyHandle> Keys);
 public:
 
 	/**
 	 * Toggle the selection of the specified point handle, changing the selection type if necessary.
 	 */
-	void Toggle(FCurvePointHandle InHandle);
+	UE_API void Toggle(FCurvePointHandle InHandle);
 
 	/**
 	 * Toggle the selection of the specified key handle, changing the selection type if necessary.
 	 */
-	void Toggle(FCurveModelID CurveID, ECurvePointType PointType, FKeyHandle KeyHandle);
+	UE_API void Toggle(FCurveModelID CurveID, ECurvePointType PointType, FKeyHandle KeyHandle);
 
 	/**
 	 * Toggle the selection of the specified key handles, changing the selection type if necessary.
 	 */
-	void Toggle(FCurveModelID CurveID, ECurvePointType PointType, TArrayView<const FKeyHandle> Keys);
+	UE_API void Toggle(FCurveModelID CurveID, ECurvePointType PointType, TArrayView<const FKeyHandle> Keys);
 
 public:
 
 	/**
 	 * Remove the specified point handle from the selection
 	 */
-	void Remove(FCurvePointHandle InHandle);
+	UE_API void Remove(FCurvePointHandle InHandle);
 
 	/**
 	 * Remove the specified key handle from the selection
 	 */
-	void Remove(FCurveModelID CurveID, ECurvePointType PointType, FKeyHandle KeyHandle);
+	UE_API void Remove(FCurveModelID CurveID, ECurvePointType PointType, FKeyHandle KeyHandle);
 
 	/**
 	 * Remove the specified key handles from the selection
 	 */
-	void Remove(FCurveModelID CurveID, ECurvePointType PointType, TArrayView<const FKeyHandle> Keys);
+	UE_API void Remove(FCurveModelID CurveID, ECurvePointType PointType, TArrayView<const FKeyHandle> Keys);
 
 	/**
 	 * Remove all key handles associated with the specified curve ID from the selection
 	 */
-	void Remove(FCurveModelID InCurveID);
+	UE_API void Remove(FCurveModelID InCurveID);
 
 	/**
 	 * Clear the selection entirely
 	 */
-	void Clear();
+	UE_API void Clear();
 
 public:
 
 	/** Serializes the SerialNumber and CurveToSelectedKeys. */
-	void SerializeSelection(FArchive& Ar);
+	UE_API void SerializeSelection(FArchive& Ar);
 
 private:
 
@@ -213,5 +215,7 @@ private:
 	TMap<FCurveModelID, FKeyHandleSet> CurveToSelectedKeys;
 	
 	/** Adds the specified key handles from the selection without validating whether CurveId and Keys exists in the owning curve editor. */
-	void AddInternal(FCurveModelID CurveID, ECurvePointType PointType, TArrayView<const FKeyHandle> Keys, bool bIncrementSerialNumber = true);
-};
\ No newline at end of file
+	UE_API void AddInternal(FCurveModelID CurveID, ECurvePointType PointType, TArrayView<const FKeyHandle> Keys, bool bIncrementSerialNumber = true);
+};
+
+#undef UE_API
diff --git a/Engine/Source/Editor/CurveEditor/Public/CurveEditorSettings.h b/Engine/Source/Editor/CurveEditor/Public/CurveEditorSettings.h
--- a/Engine/Source/Editor/CurveEditor/Public/CurveEditorSettings.h
+++ b/Engine/Source/Editor/CurveEditor/Public/CurveEditorSettings.h
@@ -15,6 +15,8 @@
 
 #include "CurveEditorSettings.generated.h"
 
+#define UE_API CURVEEDITOR_API
+
 class UClass;
 
 /** Defines visibility states for the tangents in the curve editor. */
@@ -91,8 +93,8 @@ struct FCustomColorForSpaceSwitch
 };
 
 /** Serializable options for curve editor. */
-UCLASS(config=EditorPerProjectUserSettings)
-class CURVEEDITOR_API UCurveEditorSettings : public UObject
+UCLASS(MinimalAPI, config=EditorPerProjectUserSettings)
+class UCurveEditorSettings : public UObject
 {
 public:
 	GENERATED_BODY()
@@ -101,69 +103,69 @@ public:
 	DECLARE_MULTICAST_DELEGATE(FOnAxisSnappingChanged);
 	DECLARE_MULTICAST_DELEGATE(FOnShowValueIndicatorsChanged);
 
-	UCurveEditorSettings();
+	UE_API UCurveEditorSettings();
 
 	/** Gets whether or not the curve editor auto frames the selected curves. */
-	bool GetAutoFrameCurveEditor() const;
+	UE_API bool GetAutoFrameCurveEditor() const;
 	/** Sets whether or not the curve editor auto frames the selected curves. */
-	void SetAutoFrameCurveEditor(bool InbAutoFrameCurveEditor);
+	UE_API void SetAutoFrameCurveEditor(bool InbAutoFrameCurveEditor);
 
 	/** Gets whether or not the curve editor shows key bar style curves, like for constraints and spaces. */
-	bool GetShowBars() const;
+	UE_API bool GetShowBars() const;
 	/** Sets whether or not the curve editor shows key bar style curves, like for constraints and spaces. */
-	void SetShowBars(bool InShowBars);
+	UE_API void SetShowBars(bool InShowBars);
 
 	/** Gets the number of pixels to pad input framing */
-	int32 GetFrameInputPadding() const;
+	UE_API int32 GetFrameInputPadding() const;
 	/** Sets the number of pixels to pad input framing */
-	void SetFrameInputPadding(int32 InFrameInputPadding);
+	UE_API void SetFrameInputPadding(int32 InFrameInputPadding);
 
 	/** Gets the number of pixels to pad output framing */
-	int32 GetFrameOutputPadding() const;
+	UE_API int32 GetFrameOutputPadding() const;
 	/** Sets the number of pixels to pad output framing */
-	void SetFrameOutputPadding(int32 InFrameOutputPadding);
+	UE_API void SetFrameOutputPadding(int32 InFrameOutputPadding);
 
 	/** Gets whether or not to show buffered curves in the curve editor. */
-	bool GetShowBufferedCurves() const;
+	UE_API bool GetShowBufferedCurves() const;
 	/** Sets whether or not to show buffered curves in the curve editor. */
-	void SetShowBufferedCurves(bool InbShowBufferedCurves);
+	UE_API void SetShowBufferedCurves(bool InbShowBufferedCurves);
 
 	/** Gets whether or not to show curve tool tips in the curve editor. */
-	bool GetShowCurveEditorCurveToolTips() const;
+	UE_API bool GetShowCurveEditorCurveToolTips() const;
 	/** Sets whether or not to show curve tool tips in the curve editor. */
-	void SetShowCurveEditorCurveToolTips(bool InbShowCurveEditorCurveToolTips);
+	UE_API void SetShowCurveEditorCurveToolTips(bool InbShowCurveEditorCurveToolTips);
 
 	/** Gets the current tangent visibility. */
-	ECurveEditorTangentVisibility GetTangentVisibility() const;
+	UE_API ECurveEditorTangentVisibility GetTangentVisibility() const;
 	/** Sets the current tangent visibility. */
-	void SetTangentVisibility(ECurveEditorTangentVisibility InTangentVisibility);
+	UE_API void SetTangentVisibility(ECurveEditorTangentVisibility InTangentVisibility);
 
 	/** Get zoom in/out position (mouse position or current time). */
-	ECurveEditorZoomPosition GetZoomPosition() const;
+	UE_API ECurveEditorZoomPosition GetZoomPosition() const;
 	/** Set zoom in/out position (mouse position or current time). */
-	void SetZoomPosition(ECurveEditorZoomPosition InZoomPosition);
+	UE_API void SetZoomPosition(ECurveEditorZoomPosition InZoomPosition);
 
 	/** Get snap axis. */
-	ECurveEditorSnapAxis GetSnapAxis() const;
+	UE_API ECurveEditorSnapAxis GetSnapAxis() const;
 	/** Set snap axis. */
-	void SetSnapAxis(ECurveEditorSnapAxis InSnapAxis);
+	UE_API void SetSnapAxis(ECurveEditorSnapAxis InSnapAxis);
 
 	/** Get whether to snap the time to the currently selected key. */
-	bool GetSnapTimeToSelection() const;
+	UE_API bool GetSnapTimeToSelection() const;
 	/** Set whether to snap the time to the currently selected key. */
-	void SetSnapTimeToSelection(bool bInSnapTimeToSelection);
+	UE_API void SetSnapTimeToSelection(bool bInSnapTimeToSelection);
 
 	/** Set the selection color. */
-	void SetSelectionColor(const FLinearColor& InColor);
+	UE_API void SetSelectionColor(const FLinearColor& InColor);
 	/** Get the selection color. */
-	FLinearColor GetSelectionColor() const;
+	UE_API FLinearColor GetSelectionColor() const;
 
 	/** Get custom color for object and property if it exists, if it doesn't the optional won't be set */
-	TOptional<FLinearColor> GetCustomColor(UClass* InClass, const FString& InPropertyName) const;
+	UE_API TOptional<FLinearColor> GetCustomColor(UClass* InClass, const FString& InPropertyName) const;
 	/** Set Custom Color for the specified parameters. */
-	void SetCustomColor(UClass* InClass, const FString& InPropertyName, FLinearColor InColor);
+	UE_API void SetCustomColor(UClass* InClass, const FString& InPropertyName, FLinearColor InColor);
 	/** Delete Custom Color for the specified parameters. */
-	void DeleteCustomColor(UClass* InClass, const FString& InPropertyName);
+	UE_API void DeleteCustomColor(UClass* InClass, const FString& InPropertyName);
 	
 	/** Gets the multicast delegate which is run whenever custom colors have changed. */
 	FOnCustomColorsChanged& GetOnCustomColorsChanged() { return OnCustomColorsChangedEvent; }
@@ -173,33 +175,33 @@ public:
 	FOnShowValueIndicatorsChanged& OnShowValueIndicatorsChanged() { return OnShowValueIndicatorsChangedEvent; }
 
 	/** Get custom color for space name. Parent and World are reserved names and will be used instead of the specified control name. */
-	TOptional<FLinearColor> GetSpaceSwitchColor(const FString& InControlName) const;
+	UE_API TOptional<FLinearColor> GetSpaceSwitchColor(const FString& InControlName) const;
 	/** Set Custom Space SwitchColor for the specified control name. */
-	void SetSpaceSwitchColor(const FString& InControlName, FLinearColor InColor);
+	UE_API void SetSpaceSwitchColor(const FString& InControlName, FLinearColor InColor);
 	/** Delete Custom Space Switch Color for the specified control name. */
-	void DeleteSpaceSwitchColor(const FString& InControlName);
+	UE_API void DeleteSpaceSwitchColor(const FString& InControlName);
 
 	/** Helper function to get next random linear color*/
-	static FLinearColor GetNextRandomColor();
+	static UE_API FLinearColor GetNextRandomColor();
 
 	/** Gets the tree view width percentage */
 	float GetTreeViewWidth() const { return TreeViewWidth; }
 	/** Sets the tree view width percentage */
-	void SetTreeViewWidth(float InTreeViewWidth);
+	UE_API void SetTreeViewWidth(float InTreeViewWidth);
 
 	/** Gets how sensitive the selection marquee should be when selecting points. */
 	float GetMarqueePointSensitivity() const { return FMath::Clamp(MarqueePointSensitivity, 0.f, 1.f); }
 	/** Sets how sensitive the selection marquee should be when selecting points. */
-	void SetMarqueePointSensitivity(float InMarqueePointSensitivity);
+	UE_API void SetMarqueePointSensitivity(float InMarqueePointSensitivity);
 
 	/** @return Whether to draw a value indicator line for the minimum and maximum key in the selected key range. */
 	bool GetShowValueIndicators() const { return bShowValueIndicators; }
 	/** Sets whether to draw a value indicator line for the minimum and maximum key in the selected key range. */
-	void SetShowValueIndicators(bool bValue);
+	UE_API void SetShowValueIndicators(bool bValue);
 
 protected:
 
-	virtual void PostEditChangeProperty(FPropertyChangedEvent& PropertyChangedEvent) override;
+	UE_API virtual void PostEditChangeProperty(FPropertyChangedEvent& PropertyChangedEvent) override;
 
 	UPROPERTY( config, EditAnywhere, Category="Curve Editor" )
 	bool bAutoFrameCurveEditor;
@@ -270,3 +272,5 @@ private:
 	FOnAxisSnappingChanged OnAxisSnappingChangedEvent;
 	FOnShowValueIndicatorsChanged OnShowValueIndicatorsChangedEvent;
 };
+
+#undef UE_API
diff --git a/Engine/Source/Editor/CurveEditor/Public/CurveEditorTypes.h b/Engine/Source/Editor/CurveEditor/Public/CurveEditorTypes.h
--- a/Engine/Source/Editor/CurveEditor/Public/CurveEditorTypes.h
+++ b/Engine/Source/Editor/CurveEditor/Public/CurveEditorTypes.h
@@ -9,6 +9,8 @@
 #include "Containers/Array.h"
 #include "Templates/SharedPointer.h"
 
+#define UE_API CURVEEDITOR_API
+
 /**
  * Enum for representing the type of a key point in the curve editor
  */
@@ -45,7 +47,7 @@ ENUM_CLASS_FLAGS(ECurveEditorViewID);
 /**
  * A unique identifier for a curve model existing on a curve editor
  */
-struct CURVEEDITOR_API FCurveModelID
+struct FCurveModelID
 {
 	/** Default constructor is needed by e.g. TArray::operator<< */
 	FCurveModelID() : ID(0) {}
@@ -53,7 +55,7 @@ struct CURVEEDITOR_API FCurveModelID
 	/**
 	 * Generate a new curve model ID
 	 */
-	static FCurveModelID Unique();
+	static UE_API FCurveModelID Unique();
 
 	/**
 	 * @return Whether this ID is valid (was constructed with Unique)
@@ -167,4 +169,6 @@ protected:
 
 	friend class FCurveEditorTree;
 	uint32 Value;
-};
\ No newline at end of file
+};
+
+#undef UE_API
diff --git a/Engine/Source/Editor/CurveEditor/Public/CurveModel.h b/Engine/Source/Editor/CurveEditor/Public/CurveModel.h
--- a/Engine/Source/Editor/CurveEditor/Public/CurveModel.h
+++ b/Engine/Source/Editor/CurveEditor/Public/CurveModel.h
@@ -20,6 +20,8 @@
 #include "Templates/UniquePtr.h"
 #include "UObject/UnrealType.h"
 
+#define UE_API CURVEEDITOR_API
+
 class FCurveEditor;
 class FCurveEditorAxis;
 class FName;
@@ -49,7 +51,7 @@ namespace UE::CurveEditor
 /**
  * Class that models an underlying curve data structure through a generic abstraction that the curve editor understands.
  */
-class CURVEEDITOR_API FCurveModel
+class FCurveModel
 {
 public:
 
@@ -209,7 +211,7 @@ public:
 	 * @param OutPreviousKeyHandle The key with time <= InTime that is closest to InTime. Unset if there is no such key.
 	 * @param OutNextKeyHandle The key with time >= InTime that is closest to InTime. Unset if there is no such key.
 	 */
-	virtual void GetClosestKeysTo(double InTime, TOptional<FKeyHandle>& OutPreviousKeyHandle, TOptional<FKeyHandle>& OutNextKeyHandle) const;
+	UE_API virtual void GetClosestKeysTo(double InTime, TOptional<FKeyHandle>& OutPreviousKeyHandle, TOptional<FKeyHandle>& OutNextKeyHandle) const;
 
 	/**
 	 * Get the interpolation mode to use at a specified time
@@ -396,12 +398,12 @@ public:
 	/**
 	 * Helper function for assigning a the same attributes to a number of keys
 	 */
-	void SetKeyAttributes(TArrayView<const FKeyHandle> InKeys, const FKeyAttributes& InAttributes, EPropertyChangeType::Type ChangeType = EPropertyChangeType::Unspecified);
+	UE_API void SetKeyAttributes(TArrayView<const FKeyHandle> InKeys, const FKeyAttributes& InAttributes, EPropertyChangeType::Type ChangeType = EPropertyChangeType::Unspecified);
 
 	/**
 	 * Helper function for adding a single key to this curve
 	 */
-	TOptional<FKeyHandle> AddKey(const FKeyPosition& NewKeyPosition, const FKeyAttributes& InAttributes);
+	UE_API TOptional<FKeyHandle> AddKey(const FKeyPosition& NewKeyPosition, const FKeyAttributes& InAttributes);
 
 
 	/**
@@ -655,3 +657,5 @@ private:
 	 */
 	TOptional<FCurveModelID> CurveId;
 };
+
+#undef UE_API
diff --git a/Engine/Source/Editor/CurveEditor/Public/Filters/CurveEditorBakeFilter.h b/Engine/Source/Editor/CurveEditor/Public/Filters/CurveEditorBakeFilter.h
--- a/Engine/Source/Editor/CurveEditor/Public/Filters/CurveEditorBakeFilter.h
+++ b/Engine/Source/Editor/CurveEditor/Public/Filters/CurveEditorBakeFilter.h
@@ -13,6 +13,8 @@
 
 #include "CurveEditorBakeFilter.generated.h"
 
+#define UE_API CURVEEDITOR_API
+
 class FCurveEditor;
 class UObject;
 struct FCurveModelID;
@@ -30,8 +32,8 @@ struct FCurveEditorBakeFilterRange
 	FFrameNumber Max = FFrameNumber(10);
 };
 
-UCLASS(DisplayName="Bake")
-class CURVEEDITOR_API UCurveEditorBakeFilter : public UCurveEditorFilterBase
+UCLASS(MinimalAPI, DisplayName="Bake")
+class UCurveEditorBakeFilter : public UCurveEditorFilterBase
 {
 	GENERATED_BODY()
 public:
@@ -46,9 +48,9 @@ public:
 		BakeInterval = FFrameNumber(1);
 	}
 protected:
-	virtual void InitializeFilter_Impl(TSharedRef<FCurveEditor> InCurveEditoro) override;
-	virtual void ApplyFilter_Impl(TSharedRef<FCurveEditor> InCurveEditor, const TMap<FCurveModelID, FKeyHandleSet>& InKeysToOperateOn, TMap<FCurveModelID, FKeyHandleSet>& OutKeysToSelect) override;
-	virtual bool CanApplyFilter_Impl(TSharedRef<FCurveEditor> InCurveEditor) override;
+	UE_API virtual void InitializeFilter_Impl(TSharedRef<FCurveEditor> InCurveEditoro) override;
+	UE_API virtual void ApplyFilter_Impl(TSharedRef<FCurveEditor> InCurveEditor, const TMap<FCurveModelID, FKeyHandleSet>& InKeysToOperateOn, TMap<FCurveModelID, FKeyHandleSet>& OutKeysToSelect) override;
+	UE_API virtual bool CanApplyFilter_Impl(TSharedRef<FCurveEditor> InCurveEditor) override;
 
 public:
 
@@ -83,4 +85,6 @@ public:
 private:
 	FFrameRate InitialDisplayRate;
 	FFrameRate InitialTickResolution;
-};
\ No newline at end of file
+};
+
+#undef UE_API
diff --git a/Engine/Source/Editor/CurveEditor/Public/Filters/CurveEditorEulerFilter.h b/Engine/Source/Editor/CurveEditor/Public/Filters/CurveEditorEulerFilter.h
--- a/Engine/Source/Editor/CurveEditor/Public/Filters/CurveEditorEulerFilter.h
+++ b/Engine/Source/Editor/CurveEditor/Public/Filters/CurveEditorEulerFilter.h
@@ -11,16 +11,18 @@
 
 #include "CurveEditorEulerFilter.generated.h"
 
+#define UE_API CURVEEDITOR_API
+
 class FCurveEditor;
 class UObject;
 struct FCurveModelID;
 struct FKeyHandleSet;
 
-UCLASS(DisplayName = "Euler", meta = (
+UCLASS(MinimalAPI, DisplayName = "Euler", meta = (
 	CurveEditorLabel = "Euler Filter",
 	CurveEditorDescription = "Uses the Euler filter to correct jumps in rotation caused by angle wrapping or gimbal lock. For example, instead of wrapping around with a large jump from 178° to -178°, it smoothly interpolates a 4° transition, ensuring natural and continuous movement.")
 	)
-class CURVEEDITOR_API UCurveEditorEulerFilter : public UCurveEditorFilterBase
+class UCurveEditorEulerFilter : public UCurveEditorFilterBase
 {
 	GENERATED_BODY()
 public:
@@ -28,7 +30,9 @@ public:
 protected:
 
 	//~ Begin UCurveEditorFilterBase Interface
-	virtual void ApplyFilter_Impl(TSharedRef<FCurveEditor> InCurveEditor, const TMap<FCurveModelID, FKeyHandleSet>& InKeysToOperateOn, TMap<FCurveModelID, FKeyHandleSet>& OutKeysToSelect) override;
-	virtual bool CanApplyFilter_Impl(TSharedRef<FCurveEditor> InCurveEditor) override;
+	UE_API virtual void ApplyFilter_Impl(TSharedRef<FCurveEditor> InCurveEditor, const TMap<FCurveModelID, FKeyHandleSet>& InKeysToOperateOn, TMap<FCurveModelID, FKeyHandleSet>& OutKeysToSelect) override;
+	UE_API virtual bool CanApplyFilter_Impl(TSharedRef<FCurveEditor> InCurveEditor) override;
 	//~ End UCurveEditorFilterBase Interface
-};
\ No newline at end of file
+};
+
+#undef UE_API
diff --git a/Engine/Source/Editor/CurveEditor/Public/Filters/CurveEditorFilterBase.h b/Engine/Source/Editor/CurveEditor/Public/Filters/CurveEditorFilterBase.h
--- a/Engine/Source/Editor/CurveEditor/Public/Filters/CurveEditorFilterBase.h
+++ b/Engine/Source/Editor/CurveEditor/Public/Filters/CurveEditorFilterBase.h
@@ -8,6 +8,8 @@
 #include "CurveEditorTypes.h"
 #include "CurveEditorFilterBase.generated.h"
 
+#define UE_API CURVEEDITOR_API
+
 class FCurveEditor;
 class FText;
 struct FSlateIcon;
@@ -23,22 +25,22 @@ template <typename T> class TSubclassOf;
 * The CDO is directly edited in the Curve Editor Filter panel so if you wish to save user settings between runs you can
 * add the "config" property to your UPROPERTY() and it will automatically get saved.
 */
-UCLASS(BlueprintType, config=EditorSettings)
-class CURVEEDITOR_API UCurveEditorFilterBase : public UObject
+UCLASS(MinimalAPI, BlueprintType, config=EditorSettings)
+class UCurveEditorFilterBase : public UObject
 {
 	GENERATED_BODY()
 
 public:
 
 	/** @return The value of the UCLASS(meta = (CurveEditorLabel = "value")) */
-	static FText GetLabel(const TSubclassOf<UCurveEditorFilterBase>& InClass);
+	static UE_API FText GetLabel(const TSubclassOf<UCurveEditorFilterBase>& InClass);
 	/** @return The value of the UCLASS(meta = (CurveEditorDescription = "value")) */
-	static FText GetDescription(const TSubclassOf<UCurveEditorFilterBase>& InClass);
+	static UE_API FText GetDescription(const TSubclassOf<UCurveEditorFilterBase>& InClass);
 	/**
 	 * @return The icon associated with this class. Icons can be overriden in editor styles (e.g. FStarshipEditorStyle) with "ClassIcon.ClassName".
 	 * @see FSlateIconFinder::FindIconForClass
 	 */
-	static FSlateIcon GetIcon(const TSubclassOf<UCurveEditorFilterBase>& InClass);
+	static UE_API FSlateIcon GetIcon(const TSubclassOf<UCurveEditorFilterBase>& InClass);
 
 	/**
 	 * Initialize this filter and any parameters it depends on
@@ -159,5 +161,7 @@ protected:
 	virtual void ApplyFilter_Impl(TSharedRef<FCurveEditor> InCurveEditor, const TMap<FCurveModelID, FKeyHandleSet>& InKeysToOperateOn, TMap<FCurveModelID, FKeyHandleSet>& OutKeysToSelect) {}
 
 	/** Implementation function for determining whether this filter can be applied */
-	virtual bool CanApplyFilter_Impl(TSharedRef<FCurveEditor> InCurveEditor);
+	UE_API virtual bool CanApplyFilter_Impl(TSharedRef<FCurveEditor> InCurveEditor);
 };
+
+#undef UE_API
diff --git a/Engine/Source/Editor/CurveEditor/Public/Filters/CurveEditorSmartSnapFilter.h b/Engine/Source/Editor/CurveEditor/Public/Filters/CurveEditorSmartSnapFilter.h
--- a/Engine/Source/Editor/CurveEditor/Public/Filters/CurveEditorSmartSnapFilter.h
+++ b/Engine/Source/Editor/CurveEditor/Public/Filters/CurveEditorSmartSnapFilter.h
@@ -6,20 +6,24 @@
 #include "CurveEditorFilterBase.h"
 #include "CurveEditorSmartSnapFilter.generated.h"
 
+#define UE_API CURVEEDITOR_API
+
 /**
  * 
  */
-UCLASS(DisplayName = "Smart Snap", meta = (
+UCLASS(MinimalAPI, DisplayName = "Smart Snap", meta = (
 	CurveEditorLabel = "Smart Snap",
 	CurveEditorDescription = "Snaps selected keys to the closest whole frame.\nThe key is placed where the curve intersects with the frame (imagine vertical line going down from the frame).")
 	)
-class CURVEEDITOR_API UCurveEditorSmartSnapFilter : public UCurveEditorFilterBase
+class UCurveEditorSmartSnapFilter : public UCurveEditorFilterBase
 {
 	GENERATED_BODY()
 protected:
 
 	//~ Begin UCurveEditorFilterBase Interface
-	virtual void ApplyFilter_Impl(TSharedRef<FCurveEditor> InCurveEditor, const TMap<FCurveModelID, FKeyHandleSet>& InKeysToOperateOn, TMap<FCurveModelID, FKeyHandleSet>& OutKeysToSelect) override;
-	virtual bool CanApplyFilter_Impl(TSharedRef<FCurveEditor> InCurveEditor) override;
+	UE_API virtual void ApplyFilter_Impl(TSharedRef<FCurveEditor> InCurveEditor, const TMap<FCurveModelID, FKeyHandleSet>& InKeysToOperateOn, TMap<FCurveModelID, FKeyHandleSet>& OutKeysToSelect) override;
+	UE_API virtual bool CanApplyFilter_Impl(TSharedRef<FCurveEditor> InCurveEditor) override;
 	//~ End UCurveEditorFilterBase Interface
 };
+
+#undef UE_API
diff --git a/Engine/Source/Editor/CurveEditor/Public/Filters/PromotedFilterCommandBinder.h b/Engine/Source/Editor/CurveEditor/Public/Filters/PromotedFilterCommandBinder.h
--- a/Engine/Source/Editor/CurveEditor/Public/Filters/PromotedFilterCommandBinder.h
+++ b/Engine/Source/Editor/CurveEditor/Public/Filters/PromotedFilterCommandBinder.h
@@ -6,6 +6,8 @@
 #include "Templates/UnrealTemplate.h"
 #include "Templates/SharedPointer.h"
 
+#define UE_API CURVEEDITOR_API
+
 class FCurveEditor;
 class FUICommandInfo;
 class FUICommandList;
@@ -18,7 +20,7 @@ class FPromotedFilterContainer;
  * Binds and unbinds commands that are created by FPromotedFilterContainer.
  * When a filter is added to FPromotedFilterContainer, it creates a command to invoke that filter. This class is responsible to bind / unbind it.
  */
-class CURVEEDITOR_API FPromotedFilterCommandBinder : public FNoncopyable
+class FPromotedFilterCommandBinder : public FNoncopyable
 {
 public:
 	
@@ -27,12 +29,12 @@ public:
 	 * @param InCommandList The command list to bind commands to.
 	 * @param InCurveEditor The curve editor on which the filters will be applied.
 	 */
-	explicit FPromotedFilterCommandBinder(
+	UE_API explicit FPromotedFilterCommandBinder(
 		const TSharedRef<FPromotedFilterContainer>& InContainer,
 		const TSharedRef<FUICommandList>& InCommandList,
 		const TSharedRef<FCurveEditor>& InCurveEditor
 		);
-	~FPromotedFilterCommandBinder();
+	UE_API ~FPromotedFilterCommandBinder();
 
 private:
 
@@ -43,13 +45,14 @@ private:
 	/** Passed to filter as argument. */
 	const TWeakPtr<FCurveEditor> CurveEditor;
 	
-	void OnFilterAdded(UCurveEditorFilterBase& InFilter, const TSharedRef<FUICommandInfo>& InCommand) const;
-	void OnFilterRemoved(UCurveEditorFilterBase& InFilter, const TSharedRef<FUICommandInfo>& InCommand) const;
+	UE_API void OnFilterAdded(UCurveEditorFilterBase& InFilter, const TSharedRef<FUICommandInfo>& InCommand) const;
+	UE_API void OnFilterRemoved(UCurveEditorFilterBase& InFilter, const TSharedRef<FUICommandInfo>& InCommand) const;
 
-	void MapAction(UCurveEditorFilterBase& InFilter, const TSharedRef<FUICommandInfo>& InCommand, FUICommandList& CommandListPin) const;
+	UE_API void MapAction(UCurveEditorFilterBase& InFilter, const TSharedRef<FUICommandInfo>& InCommand, FUICommandList& CommandListPin) const;
 	
-	void ApplyFilter(UCurveEditorFilterBase* InFilter, TWeakPtr<FUICommandInfo> Command) const;
-	bool CanApplyFilter(UCurveEditorFilterBase* InFilter, TWeakPtr<FUICommandInfo> Command) const;
+	UE_API void ApplyFilter(UCurveEditorFilterBase* InFilter, TWeakPtr<FUICommandInfo> Command) const;
+	UE_API bool CanApplyFilter(UCurveEditorFilterBase* InFilter, TWeakPtr<FUICommandInfo> Command) const;
 };
 }
 
+#undef UE_API
diff --git a/Engine/Source/Editor/CurveEditor/Public/Filters/PromotedFilterContainer.h b/Engine/Source/Editor/CurveEditor/Public/Filters/PromotedFilterContainer.h
--- a/Engine/Source/Editor/CurveEditor/Public/Filters/PromotedFilterContainer.h
+++ b/Engine/Source/Editor/CurveEditor/Public/Filters/PromotedFilterContainer.h
@@ -11,6 +11,8 @@
 
 #include <type_traits>
 
+#define UE_API CURVEEDITOR_API
+
 class FBindingContext;
 class FCurveEditor;
 class FMenuBuilder;
@@ -29,7 +31,7 @@ namespace UE::CurveEditor
  * filters to the toolbar via UI), this class should be extended to persistently save the user-specified filter settings,  e.g. the FGaussianParams
  * for the UCurveEditorGaussianFilter, etc.
  */
-class CURVEEDITOR_API FPromotedFilterContainer
+class FPromotedFilterContainer
 	: public FGCObject
 	, public FNoncopyable
 {
@@ -38,22 +40,22 @@ public:
 	/**
 	 * @param InContextName The ID to give the underlying FBindingContext. Must be globally unique.
 	 */
-	explicit FPromotedFilterContainer(FName InContextName);
-	virtual ~FPromotedFilterContainer() override;
+	UE_API explicit FPromotedFilterContainer(FName InContextName);
+	UE_API virtual ~FPromotedFilterContainer() override;
 
 	/** Appends the filters to InToolbarBuilder. */
-	void AppendToBuilder(FToolBarBuilder& InToolBarBuilder, const FMenuEntryResizeParams& InResizeParams = {}) const;
+	UE_API void AppendToBuilder(FToolBarBuilder& InToolBarBuilder, const FMenuEntryResizeParams& InResizeParams = {}) const;
 	/** Appends the filters to InMenuBuilder. */
-	void AppendToBuilder(FMenuBuilder& InMenuBuilder) const;
+	UE_API void AppendToBuilder(FMenuBuilder& InMenuBuilder) const;
 	
 	/**
 	 * Promotes a filter.
 	 * For simplicity, each class can only have one instance promoted (e.g. you cannot call AddInstance with two different UCurveEditorEulerFilter instances).
 	 */
-	void AddInstance(UCurveEditorFilterBase& InFilter);
+	UE_API void AddInstance(UCurveEditorFilterBase& InFilter);
 	
 	/** Removes a filter instance. You can pass in a CDO, too. */
-	void RemoveInstance(UCurveEditorFilterBase& InFilter);
+	UE_API void RemoveInstance(UCurveEditorFilterBase& InFilter);
 
 	template<typename TCallback> requires std::is_invocable_v<TCallback, UCurveEditorFilterBase&, const TSharedRef<FUICommandInfo>&>
 	void ForEachFilter(TCallback&& Callback);
@@ -65,7 +67,7 @@ public:
 	FFilterDelegate& OnFilterRemoved() { return OnFilterRemovedDelegate; }
 
 	//~ Begin FGCObject Interface
-	virtual void AddReferencedObjects(FReferenceCollector& Collector) override;
+	UE_API virtual void AddReferencedObjects(FReferenceCollector& Collector) override;
 	virtual FString GetReferencerName() const override { return TEXT("FPromotedFilterContainer"); }
 	//~ Begin FGCObject Interface
 
@@ -99,8 +101,8 @@ private:
 	/** The filters that have been promoted. */
 	TArray<FFilterData> PromotedFilters;
 
-	int32 IndexOf(const UClass* InFilterClass) const;
-	void RemoveAtInternal(int32 InIndex);
+	UE_API int32 IndexOf(const UClass* InFilterClass) const;
+	UE_API void RemoveAtInternal(int32 InIndex);
 };
 
 template <typename TCallback> requires std::is_invocable_v<TCallback, UCurveEditorFilterBase&, const TSharedRef<FUICommandInfo>&>
@@ -113,3 +115,4 @@ void FPromotedFilterContainer::ForEachFilter(TCallback&& Callback)
 }
 }
 
+#undef UE_API
diff --git a/Engine/Source/Editor/CurveEditor/Public/Filters/SCurveEditorFilterPanel.h b/Engine/Source/Editor/CurveEditor/Public/Filters/SCurveEditorFilterPanel.h
--- a/Engine/Source/Editor/CurveEditor/Public/Filters/SCurveEditorFilterPanel.h
+++ b/Engine/Source/Editor/CurveEditor/Public/Filters/SCurveEditorFilterPanel.h
@@ -9,6 +9,8 @@
 #include "Widgets/DeclarativeSyntaxSupport.h"
 #include "Widgets/SCompoundWidget.h"
 
+#define UE_API CURVEEDITOR_API
+
 class FCurveEditor;
 class FTabManager;
 class IDetailsView;
@@ -16,7 +18,7 @@ class SWindow;
 class UClass;
 class UCurveEditorFilterBase;
 
-class CURVEEDITOR_API SCurveEditorFilterPanel : public SCompoundWidget
+class SCurveEditorFilterPanel : public SCompoundWidget
 {
 public:
 	SLATE_BEGIN_ARGS(SCurveEditorFilterPanel)
@@ -24,15 +26,15 @@ public:
 		
 	SLATE_END_ARGS()
 
-	void Construct(const FArguments& InArgs, TSharedRef<FCurveEditor> InCurveEditor, UClass* DefaultFilterClass);
-	void SetFilterClass(UClass* InClass);
+	UE_API void Construct(const FArguments& InArgs, TSharedRef<FCurveEditor> InCurveEditor, UClass* DefaultFilterClass);
+	UE_API void SetFilterClass(UClass* InClass);
 
 public:
 	/** Call this to request opening a window containing this panel. */
-	static TSharedPtr<SCurveEditorFilterPanel> OpenDialog(TSharedPtr<SWindow> RootWindow, TSharedRef<FCurveEditor> InHostCurveEditor, TSubclassOf<UCurveEditorFilterBase> DefaultFilterClass);
+	static UE_API TSharedPtr<SCurveEditorFilterPanel> OpenDialog(TSharedPtr<SWindow> RootWindow, TSharedRef<FCurveEditor> InHostCurveEditor, TSubclassOf<UCurveEditorFilterBase> DefaultFilterClass);
 	
 	/** Closes the dialog if there is one open. */
-	static void CloseDialog();
+	static UE_API void CloseDialog();
 
 	/** The details view for filter class properties */
 	TSharedPtr<class IDetailsView> GetDetailsView() const { return DetailsView; }
@@ -41,9 +43,9 @@ public:
 	FSimpleDelegate OnFilterClassChanged;
 
 protected:
-	FReply OnApplyClicked();
-	bool CanApplyFilter() const;
-	FText GetCurrentFilterText() const;
+	UE_API FReply OnApplyClicked();
+	UE_API bool CanApplyFilter() const;
+	UE_API FText GetCurrentFilterText() const;
 private:
 	/** Weak pointer to the curve editor which created this filter. */
 	TWeakPtr<FCurveEditor> WeakCurveEditor;
@@ -52,5 +54,7 @@ private:
 	TSharedPtr<IDetailsView> DetailsView;
 
 	/** Singleton for the pop-up window. */
-	static TWeakPtr<SWindow> ExistingFilterWindow;
-};
\ No newline at end of file
+	static UE_API TWeakPtr<SWindow> ExistingFilterWindow;
+};
+
+#undef UE_API
diff --git a/Engine/Source/Editor/CurveEditor/Public/ICurveEditorCurveCachePool.h b/Engine/Source/Editor/CurveEditor/Public/ICurveEditorCurveCachePool.h
--- a/Engine/Source/Editor/CurveEditor/Public/ICurveEditorCurveCachePool.h
+++ b/Engine/Source/Editor/CurveEditor/Public/ICurveEditorCurveCachePool.h
@@ -9,7 +9,7 @@ struct FCurveDrawParams;
 
 namespace UE::CurveEditor
 {
-	class CURVEEDITOR_API ICurveEditorCurveCachePool
+	class ICurveEditorCurveCachePool
 	{
 	public:
 		/** Draws the cached curves. Called when all curves were drawn to the curve cache pool. */
diff --git a/Engine/Source/Editor/CurveEditor/Public/ICurveEditorToolExtension.h b/Engine/Source/Editor/CurveEditor/Public/ICurveEditorToolExtension.h
--- a/Engine/Source/Editor/CurveEditor/Public/ICurveEditorToolExtension.h
+++ b/Engine/Source/Editor/CurveEditor/Public/ICurveEditorToolExtension.h
@@ -20,7 +20,7 @@ struct FSlateIcon;
 /**
 * A unique identifier for a tool in the curve editor
 */
-struct CURVEEDITOR_API FCurveEditorToolID
+struct FCurveEditorToolID
 {
 	/**
 	 * Generate a new tool ID
@@ -80,7 +80,7 @@ DECLARE_MULTICAST_DELEGATE(FOnOptionsRefresh);
 * these events can be bubbled to the rest of the Curve Editor to allow common functionality of selecting/deselecting
 * keys, panning, etc.
 */
-class CURVEEDITOR_API ICurveEditorToolExtension
+class ICurveEditorToolExtension
 {
 public:
 	ICurveEditorToolExtension()
@@ -147,4 +147,4 @@ public:
 
 protected:
 	FCurveEditorToolID ToolID;
-};
\ No newline at end of file
+};
diff --git a/Engine/Source/Editor/CurveEditor/Public/KeyBarCurveModel.h b/Engine/Source/Editor/CurveEditor/Public/KeyBarCurveModel.h
--- a/Engine/Source/Editor/CurveEditor/Public/KeyBarCurveModel.h
+++ b/Engine/Source/Editor/CurveEditor/Public/KeyBarCurveModel.h
@@ -15,6 +15,8 @@
 #include "UObject/NameTypes.h"
 #include "UObject/UnrealType.h"
 
+#define UE_API CURVEEDITOR_API
+
 class FCurveEditor;
 class FMenuBuilder;
 struct FCurveEditorScreenSpace;
@@ -26,12 +28,12 @@ struct FKeyPosition;
 /**
  * Class that provides functionality to represents keys as continuous ranges.
  */
-class CURVEEDITOR_API FKeyBarCurveModel: public FCurveModel
+class FKeyBarCurveModel: public FCurveModel
 {
 public:
 
-	FKeyBarCurveModel();
-	~FKeyBarCurveModel();
+	UE_API FKeyBarCurveModel();
+	UE_API ~FKeyBarCurveModel();
 
 	virtual const void* GetCurve() const override { return nullptr; }
 	virtual void Modify() override {};
@@ -57,7 +59,7 @@ public:
 	/**
 	*  Bar Ranges for the Keys
 	*/
-	struct CURVEEDITOR_API FBarRange
+	struct FBarRange
 	{
 		/* Spaces have ifinite ranges, Constraints don't*/
 		bool bRangeIsInfinite = true;
@@ -74,3 +76,5 @@ public:
 		return Temp;
 	}
 };
+
+#undef UE_API
diff --git a/Engine/Source/Editor/CurveEditor/Public/Misc/Mirror/PositionMirrorSolver.h b/Engine/Source/Editor/CurveEditor/Public/Misc/Mirror/PositionMirrorSolver.h
--- a/Engine/Source/Editor/CurveEditor/Public/Misc/Mirror/PositionMirrorSolver.h
+++ b/Engine/Source/Editor/CurveEditor/Public/Misc/Mirror/PositionMirrorSolver.h
@@ -7,6 +7,8 @@
 #include "CurveEditorTypes.h"
 #include "UniformMirrorSolver.h"
 
+#define UE_API CURVEEDITOR_API
+
 class FCurveEditor;
 template<typename OptionalType> struct TOptional;
 
@@ -29,7 +31,7 @@ namespace UE::CurveEditor
  * 
  * @note This class only recalculates key positions. To adjust the tangents of the keys use e.g. FTangentMirrorSolver.
  */
-class CURVEEDITOR_API FPositionMirrorSolver
+class FPositionMirrorSolver
 {
 public:
 	
@@ -41,7 +43,7 @@ public:
 	 * @param InStartY The height of the line that is moved. It determines how the tangents are interpolated.
 	 * @param InMiddlePointY The height of the line on which the points are moved.
 	 */
-	explicit FPositionMirrorSolver(
+	UE_API explicit FPositionMirrorSolver(
 		FCurveEditor& InCurveEditor UE_LIFETIMEBOUND,
 		double InStartY, double InMiddlePointY
 		);
@@ -57,10 +59,10 @@ public:
 	 * @param InPositions The positions of the keys, if you have them already (optimization). If not, pass in empty and the key positions are determined.
 	 * @return Whether any of the keys can be interpolated (only user specified tangents can be interpolated).
 	 */
-	bool AddKeyPositions(const FCurveModelID& InCurveId, TArray<FKeyHandle> InKeys, TArray<FKeyPosition> InPositions = {});
+	UE_API bool AddKeyPositions(const FCurveModelID& InCurveId, TArray<FKeyHandle> InKeys, TArray<FKeyPosition> InPositions = {});
 	
 	/** Recomputes the tangents angles and updates the key attributes. */
-	void OnMoveEdge(double InDraggedEdgeHeight);
+	UE_API void OnMoveEdge(double InDraggedEdgeHeight);
 	
 private:
 	
@@ -90,3 +92,4 @@ private:
 };
 }
 
+#undef UE_API
diff --git a/Engine/Source/Editor/CurveEditor/Public/Misc/Mirror/TangentSelectionFlattener.h b/Engine/Source/Editor/CurveEditor/Public/Misc/Mirror/TangentSelectionFlattener.h
--- a/Engine/Source/Editor/CurveEditor/Public/Misc/Mirror/TangentSelectionFlattener.h
+++ b/Engine/Source/Editor/CurveEditor/Public/Misc/Mirror/TangentSelectionFlattener.h
@@ -6,6 +6,8 @@
 #include "CurveEditor.h"
 #include "MultiCurveMirrorUtils.h"
 
+#define UE_API CURVEEDITOR_API
+
 class FCurveEditor;
 
 namespace UE::CurveEditor
@@ -17,7 +19,7 @@ namespace UE::CurveEditor
  * This is useful if you have a selection of keys, apply some operation that will move the keys, and want to flatten the curves.
  * Example use case is e.g. tweening.
  */
-class CURVEEDITOR_API FTangentSelectionFlattener
+class FTangentSelectionFlattener
 {
 public:
 
@@ -25,17 +27,17 @@ public:
 	explicit FTangentSelectionFlattener(const FCurveEditor& InCurveEditor) { ResetFromSelection(InCurveEditor); }
 
 	/** Inits the tangent data from the curve editor's selection. */
-	bool ResetFromSelection(const FCurveEditor& InCurveEditor);
+	UE_API bool ResetFromSelection(const FCurveEditor& InCurveEditor);
 
 	/** @return Whether any interpolatable tangents were added (only tangents that are RCTM_User or RCTM_Break). */
-	bool AddTangents(const FCurveEditor& InCurveEditor, const FCurveModelID& InCurveId, TConstArrayView<FKeyHandle> InKeys);
+	UE_API bool AddTangents(const FCurveEditor& InCurveEditor, const FCurveModelID& InCurveId, TConstArrayView<FKeyHandle> InKeys);
 
 	/**
 	 * Call in response to the keys having moved. Recomputes user specified tangents and flattens them.
 	 * @param InCurveEditor Used to update key positions
 	 * @param bTopHasCrossedBottomEdge If since this was initialized, the top edge has moved across the bottom edge.
 	 */
-	void ComputeMirroringParallel(const FCurveEditor& InCurveEditor, bool bTopHasCrossedBottomEdge = false);
+	UE_API void ComputeMirroringParallel(const FCurveEditor& InCurveEditor, bool bTopHasCrossedBottomEdge = false);
 
 private:
 	
@@ -43,3 +45,4 @@ private:
 };
 }
 
+#undef UE_API
diff --git a/Engine/Source/Editor/CurveEditor/Public/Modification/CurveEditorCommandChange.h b/Engine/Source/Editor/CurveEditor/Public/Modification/CurveEditorCommandChange.h
--- a/Engine/Source/Editor/CurveEditor/Public/Modification/CurveEditorCommandChange.h
+++ b/Engine/Source/Editor/CurveEditor/Public/Modification/CurveEditorCommandChange.h
@@ -1,28 +1,31 @@
-﻿// Copyright Epic Games, Inc. All Rights Reserved.
+// Copyright Epic Games, Inc. All Rights Reserved.
 
 #pragma once
 
 #include "Misc/Change.h"
 #include "Templates/SharedPointerFwd.h"
 
+#define UE_API CURVEEDITOR_API
+
 class FCurveEditor;
 
 namespace UE::CurveEditor
 {
 /** Base class for actions that support undo / redo in the curve editor. */
-class CURVEEDITOR_API FCurveEditorCommandChange : public FCommandChange
+class FCurveEditorCommandChange : public FCommandChange
 {
 public:
 
 	//~ Begin FChange Interface
-	virtual bool HasExpired(UObject* InObject) const override;
-	virtual FString ToString() const override;
+	UE_API virtual bool HasExpired(UObject* InObject) const override;
+	UE_API virtual FString ToString() const override;
 	//~ End FChange Interface
 
 protected:
 
 	/** @return The curve editor associated with Object. */
-	static TSharedPtr<FCurveEditor> GetCurveEditor(UObject* InObject);
+	static UE_API TSharedPtr<FCurveEditor> GetCurveEditor(UObject* InObject);
 };
 }
 
+#undef UE_API
diff --git a/Engine/Source/Editor/CurveEditor/Public/Modification/TransactionManager.h b/Engine/Source/Editor/CurveEditor/Public/Modification/TransactionManager.h
--- a/Engine/Source/Editor/CurveEditor/Public/Modification/TransactionManager.h
+++ b/Engine/Source/Editor/CurveEditor/Public/Modification/TransactionManager.h
@@ -7,17 +7,19 @@
 #include "UObject/GCObject.h"
 #include "UObject/ObjectPtr.h"
 
+#define UE_API CURVEEDITOR_API
+
 class FCurveEditor;
 class UCurveEditorTransactionObject;
 
 namespace UE::CurveEditor
 {
 /** Interacts with the transaction system to enable command-based actions for curve editor. */
-class CURVEEDITOR_API FTransactionManager : public FNoncopyable, public FGCObject
+class FTransactionManager : public FNoncopyable, public FGCObject
 {
 public:
 
-	explicit FTransactionManager(TWeakPtr<FCurveEditor> InCurveEditor);
+	UE_API explicit FTransactionManager(TWeakPtr<FCurveEditor> InCurveEditor);
 	
 	/**
 	 * Appends a change to the current transaction.
@@ -28,15 +30,15 @@ public:
 	 * 
 	 * @param InChange The change to associate with the currently open transaction
 	 */
-	void AppendChange(TUniquePtr<FCurveEditorCommandChange> InChange) const;
+	UE_API void AppendChange(TUniquePtr<FCurveEditorCommandChange> InChange) const;
 
 	DECLARE_MULTICAST_DELEGATE_OneParam(FOnCommandAppended, const FCurveEditorCommandChange&);
 	/** Invoked after a command has been appended to the undo stack. */
 	FOnCommandAppended& OnCommandAppended() { return OnCommandAppendedDelegate; }
 
 	//~ Begin FGCObject Interface
-	virtual void AddReferencedObjects(FReferenceCollector& Collector) override;
-	virtual FString GetReferencerName() const override;
+	UE_API virtual void AddReferencedObjects(FReferenceCollector& Collector) override;
+	UE_API virtual FString GetReferencerName() const override;
 	//~ End FGCObject Interface
 
 private:
@@ -49,3 +51,4 @@ private:
 };
 }
 
+#undef UE_API
diff --git a/Engine/Source/Editor/CurveEditor/Public/Modification/Utils/ScopedSelectionTransaction.h b/Engine/Source/Editor/CurveEditor/Public/Modification/Utils/ScopedSelectionTransaction.h
--- a/Engine/Source/Editor/CurveEditor/Public/Modification/Utils/ScopedSelectionTransaction.h
+++ b/Engine/Source/Editor/CurveEditor/Public/Modification/Utils/ScopedSelectionTransaction.h
@@ -5,20 +5,22 @@
 #include "CurveEditorSelection.h"
 #include "ScopedTransaction.h"
 
+#define UE_API CURVEEDITOR_API
+
 namespace UE::CurveEditor
 {
 /**
  * This util starts a transaction. When the scope ends, checks whether any changes to selection were made and adds it an undo-able action.
  * Use if you intend to change the selection.
  */
-class CURVEEDITOR_API FScopedSelectionTransaction : public FScopedTransaction
+class FScopedSelectionTransaction : public FScopedTransaction
 {
 public:
 
-	[[nodiscard]] explicit FScopedSelectionTransaction(TWeakPtr<FCurveEditor> InCurveEditor, bool bInActuallyTransact = true);
-	[[nodiscard]] explicit FScopedSelectionTransaction(TWeakPtr<FCurveEditor> InCurveEditor, const FText& InDescription, bool bInActuallyTransact = true);
+	[[nodiscard]] UE_API explicit FScopedSelectionTransaction(TWeakPtr<FCurveEditor> InCurveEditor, bool bInActuallyTransact = true);
+	[[nodiscard]] UE_API explicit FScopedSelectionTransaction(TWeakPtr<FCurveEditor> InCurveEditor, const FText& InDescription, bool bInActuallyTransact = true);
 	
-	~FScopedSelectionTransaction();
+	UE_API ~FScopedSelectionTransaction();
 
 private:
 
@@ -29,3 +31,4 @@ private:
 };
 }
 
+#undef UE_API
