- Support plugin Verse paths with folders in Content Browser

#rnx

#Synced-CL 51330672

diff --git a/Engine/Source/Editor/ContentBrowserData/Private/ContentBrowserDataSubsystem.cpp b/Engine/Source/Editor/ContentBrowserData/Private/ContentBrowserDataSubsystem.cpp
--- a/Engine/Source/Editor/ContentBrowserData/Private/ContentBrowserDataSubsystem.cpp
+++ b/Engine/Source/Editor/ContentBrowserData/Private/ContentBrowserDataSubsystem.cpp
@@ -1078,21 +1078,7 @@ void UContentBrowserDataSubsystem::ConvertInternalPathToVirtual(const FStringVie
 		{
 			if (UsePluginVersePath(Plugin.ToSharedRef()))
 			{
-			#if 0
-				// @fixme The semantically correct thing to do would be this:
 				OutPath.Append(FPathViews::GetPath(UE::String::RemoveFromEnd(FStringView(Plugin->GetVersePath()), TEXTVIEW("/"))));
-			#else
-				// However we have an issue with multi-plugin projects to solve first
-				// The root module/plugin uses the project Verse path
-				// Non-root modules are "faking" their Verse path this way: [REDACTED_EMAIL]/project/module
-				// It's semantically invalid in Verse because that module isn't actually a sub-module of the root module
-				// This creates problems in the content browser: the root plugin being an actual folder, the virtual path 
-				// of other plugins cannot be start with it (i.e. we can't mix real folder hierarchies with virtual path hierarchies)
-				// Another reason to fix this is to support "namespaces" in a project Verse paths such as: 
-				// /owner@domain/purely/organizational/structure/project
-				OutPath.AppendChar(TCHAR('/'));
-				OutPath.Append(FPathViews::GetMountPointNameFromPath(FStringView(Plugin->GetVersePath())));
-			#endif
 				OutPath.Append(InPath.GetData(), InPath.Len());
 				return;
 			}
diff --git a/Engine/Source/Editor/ContentBrowserData/Private/ContentBrowserDataUtils.cpp b/Engine/Source/Editor/ContentBrowserData/Private/ContentBrowserDataUtils.cpp
--- a/Engine/Source/Editor/ContentBrowserData/Private/ContentBrowserDataUtils.cpp
+++ b/Engine/Source/Editor/ContentBrowserData/Private/ContentBrowserDataUtils.cpp
@@ -257,24 +257,7 @@ FText ContentBrowserDataUtils::GetFolderItemDisplayNameOverride(const FName InFo
 				UContentBrowserDataSubsystem* ContentBrowser = IContentBrowserDataModule::Get().GetSubsystem();
 				if (ContentBrowser && ContentBrowser->UsePluginVersePath(Plugin.ToSharedRef()))
 				{
-				#if 0
-					// @fixme The semantically correct thing to do would be this:
 					FolderDisplayNameOverride = FText::FromString(FPaths::GetPathLeaf(Plugin->GetVersePath()));
-				#else
-					// However we have an issue with multi-plugin projects to solve first
-					// The root module/plugin uses the project Verse path
-					// Non-root modules are "faking" their Verse path this way: [REDACTED_EMAIL]/project/module
-					// It's semantically invalid in Verse because that module isn't actually a sub-module of the root module
-					// This creates problems in the content browser: the root plugin being an actual folder, the virtual path 
-					// of other plugins cannot be start with it (i.e. we can't mix real folder hierarchies with virtual path hierarchies)
-					// Another reason to fix this is to support "namespaces" in a project Verse paths such as: 
-					// /owner@domain/purely/organizational/structure/project
-					const FStringView VersePath = Plugin->GetVersePath();
-					const int32 VersePathDomainLen = FPathViews::GetMountPointNameFromPath(VersePath, /*bOutHadClassesPrefix=*/nullptr, /*bInWithoutSlashes=*/false).Len() + 1;
-					FString VersePathRemainder(VersePath.RightChop(VersePathDomainLen));
-					VersePathRemainder.ReplaceInline(TEXT("/"), TEXT(" / "));
-					FolderDisplayNameOverride = FText::FromString(MoveTemp(VersePathRemainder));
-				#endif
 				}
 			}
 
